Actor Sequence: Check for nullptr when attempting to create director instances

diff --git a/Engine/Plugins/MovieScene/ActorSequence/Source/ActorSequence/Private/ActorSequence.cpp b/Engine/Plugins/MovieScene/ActorSequence/Source/ActorSequence/Private/ActorSequence.cpp
--- a/Engine/Plugins/MovieScene/ActorSequence/Source/ActorSequence/Private/ActorSequence.cpp
+++ b/Engine/Plugins/MovieScene/ActorSequence/Source/ActorSequence/Private/ActorSequence.cpp
@@ -167,7 +167,11 @@ void UActorSequence::UnbindPossessableObjects(const FGuid& ObjectId)
 
 UObject* UActorSequence::CreateDirectorInstance(IMovieScenePlayer& Player)
 {
-	AActor* Actor = CastChecked<AActor>(Player.GetPlaybackContext());
+	AActor* Actor = CastChecked<AActor>(Player.GetPlaybackContext(), ECastCheckedType::NullAllowed);
+	if (!Actor)
+	{
+		return nullptr;
+	}
 
 	// If this sequence is inside a blueprint, or its component's archetype is from a blueprint, we use the actor as the instace (which will be an instance of the blueprint itself)
 	if (GetTypedOuter<UBlueprintGeneratedClass>() || GetTypedOuter<UActorSequenceComponent>()->GetArchetype() != GetDefault<UActorSequenceComponent>())
