Actor Sequence: Remove redundant AutoPlay

diff --git a/Engine/Plugins/MovieScene/ActorSequence/Source/ActorSequence/Private/ActorSequenceComponent.cpp b/Engine/Plugins/MovieScene/ActorSequence/Source/ActorSequence/Private/ActorSequenceComponent.cpp
--- a/Engine/Plugins/MovieScene/ActorSequence/Source/ActorSequence/Private/ActorSequenceComponent.cpp
+++ b/Engine/Plugins/MovieScene/ActorSequence/Source/ActorSequence/Private/ActorSequenceComponent.cpp
@@ -6,7 +6,6 @@
 
 UActorSequenceComponent::UActorSequenceComponent(const FObjectInitializer& ObjectInitializer)
 	: Super(ObjectInitializer)
-	, bAutoPlay(false)
 {
 	PrimaryComponentTick.bCanEverTick = true;
 
@@ -31,7 +30,7 @@ void UActorSequenceComponent::BeginPlay()
 		SequencePlayer = NewObject<UActorSequencePlayer>(this, "SequencePlayer");
 		SequencePlayer->Initialize(Sequence, PlaybackSettings);
 
-		if (bAutoPlay)
+		if (PlaybackSettings.bAutoPlay)
 		{
 			SequencePlayer->Play();
 		}
diff --git a/Engine/Plugins/MovieScene/ActorSequence/Source/ActorSequence/Public/ActorSequenceComponent.h b/Engine/Plugins/MovieScene/ActorSequence/Source/ActorSequence/Public/ActorSequenceComponent.h
--- a/Engine/Plugins/MovieScene/ActorSequence/Source/ActorSequence/Public/ActorSequenceComponent.h
+++ b/Engine/Plugins/MovieScene/ActorSequence/Source/ActorSequence/Public/ActorSequenceComponent.h
@@ -39,9 +39,6 @@ public:
 
 protected:
 
-	UPROPERTY(EditAnywhere, Category="Playback")
-	bool bAutoPlay;
-
 	UPROPERTY(EditAnywhere, Category="Playback", meta=(ShowOnlyInnerProperties))
 	FMovieSceneSequencePlaybackSettings PlaybackSettings;
 
