Actor Sequence: Move actor sequence editor to actor sequence plugin

diff --git a/Engine/Plugins/MovieScene/ActorSequence/ActorSequence.uplugin b/Engine/Plugins/MovieScene/ActorSequence/ActorSequence.uplugin
--- a/Engine/Plugins/MovieScene/ActorSequence/ActorSequence.uplugin
+++ b/Engine/Plugins/MovieScene/ActorSequence/ActorSequence.uplugin
@@ -20,6 +20,11 @@
 			"Name" : "ActorSequence",
 			"Type" : "Runtime",
 			"LoadingPhase" : "PreDefault"
+		},
+		{
+			"Name" : "ActorSequenceEditor",
+			"Type" : "Editor",
+			"LoadingPhase" : "PostEngineInit"
 		}
 	]
 }
\ No newline at end of file
diff --git a/Engine/Plugins/MovieScene/ActorSequence/Source/ActorSequenceEditor/ActorSequenceEditor.Build.cs b/Engine/Plugins/MovieScene/ActorSequence/Source/ActorSequenceEditor/ActorSequenceEditor.Build.cs
new file mode 100644
--- /dev/null
+++ b/Engine/Plugins/MovieScene/ActorSequence/Source/ActorSequenceEditor/ActorSequenceEditor.Build.cs
@@ -0,0 +1,46 @@
+// Copyright 1998-2018 Epic Games, Inc. All Rights Reserved.
+
+using UnrealBuildTool;
+
+public class ActorSequenceEditor : ModuleRules
+{
+	public ActorSequenceEditor(ReadOnlyTargetRules Target) : base(Target)
+	{
+		PrivateDependencyModuleNames.AddRange(
+			new string[] {
+				"ActorSequence",
+				"BlueprintGraph",
+				"Core",
+				"CoreUObject",
+				"Engine",
+				"InputCore",
+				"Kismet",
+				"MovieScene",
+				"MovieSceneTools",
+				"Sequencer",
+				"EditorStyle",
+				"Slate",
+				"SlateCore",
+				"UnrealEd",
+				"TimeManagement"
+			}
+		);
+
+		PrivateIncludePaths.AddRange(
+			new string[] {
+				"ActorSequenceEditor/Private",
+			}
+		);
+
+		var DynamicModuleNames = new string[] {
+			"LevelEditor",
+			"PropertyEditor",
+		};
+
+		foreach (var Name in DynamicModuleNames)
+		{
+			PrivateIncludePathModuleNames.Add(Name);
+			DynamicallyLoadedModuleNames.Add(Name);
+		}
+	}
+}
diff --git a/Engine/Plugins/MovieScene/ActorSequence/Source/ActorSequenceEditor/Private/ActorSequenceComponentCustomization.cpp b/Engine/Plugins/MovieScene/ActorSequence/Source/ActorSequenceEditor/Private/ActorSequenceComponentCustomization.cpp
new file mode 100644
--- /dev/null
+++ b/Engine/Plugins/MovieScene/ActorSequence/Source/ActorSequenceEditor/Private/ActorSequenceComponentCustomization.cpp
@@ -0,0 +1,197 @@
+// Copyright 1998-2018 Epic Games, Inc. All Rights Reserved.
+
+#include "ActorSequenceComponentCustomization.h"
+
+#include "ActorSequence.h"
+#include "ActorSequenceComponent.h"
+#include "EditorStyleSet.h"
+#include "GameFramework/Actor.h"
+#include "IDetailsView.h"
+#include "DetailLayoutBuilder.h"
+#include "DetailCategoryBuilder.h"
+#include "DetailWidgetRow.h"
+#include "Widgets/Docking/SDockTab.h"
+#include "SSCSEditor.h"
+#include "BlueprintEditorTabs.h"
+#include "ScopedTransaction.h"
+#include "ISequencerModule.h"
+#include "Editor.h"
+#include "ActorSequenceEditorTabSummoner.h"
+#include "IPropertyUtilities.h"
+#include "Widgets/Input/SButton.h"
+
+#define LOCTEXT_NAMESPACE "ActorSequenceComponentCustomization"
+
+FName SequenceTabId("EmbeddedSequenceID");
+
+class SActorSequenceEditorWidgetWrapper : public SActorSequenceEditorWidget
+{
+public:
+	~SActorSequenceEditorWidgetWrapper()
+	{
+		GEditor->OnObjectsReplaced().Remove(OnObjectsReplacedHandle);
+	}
+
+	void Construct(const FArguments& InArgs, TWeakObjectPtr<UActorSequenceComponent> InSequenceComponent)
+	{
+		SActorSequenceEditorWidget::Construct(InArgs, nullptr);
+
+		WeakSequenceComponent = InSequenceComponent;
+		AssignSequence(GetActorSequence());
+
+		OnObjectsReplacedHandle = GEditor->OnObjectsReplaced().AddSP(this, &SActorSequenceEditorWidgetWrapper::OnObjectsReplaced);
+	}
+
+protected:
+
+	UActorSequence* GetActorSequence() const
+	{
+		UActorSequenceComponent* SequenceComponent = WeakSequenceComponent.Get();
+		return SequenceComponent ? SequenceComponent->GetSequence() : nullptr;
+	}
+
+	void OnObjectsReplaced(const TMap<UObject*, UObject*>& ReplacementMap)
+	{
+		UActorSequenceComponent* Component = WeakSequenceComponent.Get(true);
+
+		UActorSequenceComponent* NewSequenceComponent = Component ? Cast<UActorSequenceComponent>(ReplacementMap.FindRef(Component)) : nullptr;
+		if (NewSequenceComponent)
+		{
+			WeakSequenceComponent = NewSequenceComponent;
+			AssignSequence(GetActorSequence());
+		}
+	}
+
+private:
+
+	TWeakObjectPtr<UActorSequenceComponent> WeakSequenceComponent;
+	FDelegateHandle OnObjectsReplacedHandle;
+};
+
+TSharedRef<IDetailCustomization> FActorSequenceComponentCustomization::MakeInstance()
+{
+	return MakeShared<FActorSequenceComponentCustomization>();
+}
+
+void FActorSequenceComponentCustomization::CustomizeDetails(IDetailLayoutBuilder& DetailBuilder)
+{
+	PropertyUtilities = DetailBuilder.GetPropertyUtilities();
+
+	TArray<TWeakObjectPtr<UObject>> Objects;
+	DetailBuilder.GetObjectsBeingCustomized(Objects);
+	if (Objects.Num() != 1)
+	{
+		// @todo: Display other common properties / information?
+		return;
+	}
+
+	WeakSequenceComponent = Cast<UActorSequenceComponent>(Objects[0].Get());
+	if (!WeakSequenceComponent.Get())
+	{
+		return;
+	}
+
+	const IDetailsView* DetailsView = DetailBuilder.GetDetailsView();
+	TSharedPtr<FTabManager> HostTabManager = DetailsView->GetHostTabManager();
+
+	DetailBuilder.HideProperty("Sequence");
+
+	IDetailCategoryBuilder& Category = DetailBuilder.EditCategory("Sequence", FText(), ECategoryPriority::Important);
+
+	bool bIsExternalTabAlreadyOpened = false;
+
+	if (HostTabManager.IsValid() && HostTabManager->CanSpawnTab(SequenceTabId))
+	{
+		WeakTabManager = HostTabManager;
+
+		TSharedPtr<SDockTab> ExistingTab = HostTabManager->FindExistingLiveTab(SequenceTabId);
+		if (ExistingTab.IsValid())
+		{
+			UActorSequence* ThisSequence = GetActorSequence();
+
+			auto SequencerWidget = StaticCastSharedRef<SActorSequenceEditorWidget>(ExistingTab->GetContent());
+			bIsExternalTabAlreadyOpened = ThisSequence && SequencerWidget->GetSequence() == ThisSequence;
+		}
+
+		Category.AddCustomRow(FText())
+			.NameContent()
+			[
+				SNew(STextBlock)
+				.Text(LOCTEXT("SequenceValueText", "Sequence"))
+				.Font(DetailBuilder.GetDetailFont())
+			]
+			.ValueContent()
+			[
+				SNew(SButton)
+				.OnClicked(this, &FActorSequenceComponentCustomization::InvokeSequencer)
+				[
+					SNew(STextBlock)
+					.Text(bIsExternalTabAlreadyOpened ? LOCTEXT("FocusSequenceTabButtonText", "Focus Tab") : LOCTEXT("OpenSequenceTabButtonText", "Open in Tab"))
+					.Font(DetailBuilder.GetDetailFont())
+				]
+			];
+	}
+
+	// Only display an inline editor for non-blueprint sequences
+	if (GetActorSequence() && !GetActorSequence()->GetParentBlueprint() && !bIsExternalTabAlreadyOpened)
+	{
+		Category.AddCustomRow(FText())
+		.WholeRowContent()
+		.MaxDesiredWidth(TOptional<float>())
+		[
+			SAssignNew(InlineSequencer, SBox)
+			.HeightOverride(300)
+			[
+				SNew(SActorSequenceEditorWidgetWrapper, WeakSequenceComponent)
+			]
+		];
+	}
+}
+
+FReply FActorSequenceComponentCustomization::InvokeSequencer()
+{
+	TSharedPtr<FTabManager> TabManager = WeakTabManager.Pin();
+	if (TabManager.IsValid() && TabManager->CanSpawnTab(SequenceTabId))
+	{
+		TSharedRef<SDockTab> Tab = TabManager->InvokeTab(SequenceTabId);
+
+		{
+			// Set up a delegate that forces a refresh of this panel when the tab is closed to ensure we see the inline widget
+			TWeakPtr<IPropertyUtilities> WeakUtilities = PropertyUtilities;
+			auto OnClosed = [WeakUtilities](TSharedRef<SDockTab>)
+			{
+				TSharedPtr<IPropertyUtilities> PinnedPropertyUtilities = WeakUtilities.Pin();
+				if (PinnedPropertyUtilities.IsValid())
+				{
+					PinnedPropertyUtilities->EnqueueDeferredAction(FSimpleDelegate::CreateSP(PinnedPropertyUtilities.ToSharedRef(), &IPropertyUtilities::ForceRefresh));
+				}
+			};
+
+			Tab->SetOnTabClosed(SDockTab::FOnTabClosedCallback::CreateLambda(OnClosed));
+		}
+
+		// Move our inline widget content to the tab (so that we keep the existing sequencer state)
+		if (InlineSequencer.IsValid())
+		{
+			Tab->SetContent(InlineSequencer->GetChildren()->GetChildAt(0));
+			InlineSequencer->SetContent(SNullWidget::NullWidget);
+			InlineSequencer->SetVisibility(EVisibility::Collapsed);
+		}
+		else 
+		{
+			StaticCastSharedRef<SActorSequenceEditorWidget>(Tab->GetContent())->AssignSequence(GetActorSequence());
+		}
+	}
+
+	PropertyUtilities->ForceRefresh();
+
+	return FReply::Handled();
+}
+
+UActorSequence* FActorSequenceComponentCustomization::GetActorSequence() const
+{
+	UActorSequenceComponent* SequenceComponent = WeakSequenceComponent.Get();
+	return SequenceComponent ? SequenceComponent->GetSequence() : nullptr;
+}
+
+#undef LOCTEXT_NAMESPACE
diff --git a/Engine/Plugins/MovieScene/ActorSequence/Source/ActorSequenceEditor/Private/ActorSequenceComponentCustomization.h b/Engine/Plugins/MovieScene/ActorSequence/Source/ActorSequenceEditor/Private/ActorSequenceComponentCustomization.h
new file mode 100644
--- /dev/null
+++ b/Engine/Plugins/MovieScene/ActorSequence/Source/ActorSequenceEditor/Private/ActorSequenceComponentCustomization.h
@@ -0,0 +1,35 @@
+// Copyright 1998-2018 Epic Games, Inc. All Rights Reserved.
+
+#pragma once
+
+#include "CoreMinimal.h"
+#include "IDetailCustomization.h"
+#include "Widgets/Layout/SBox.h"
+#include "Input/Reply.h"
+#include "UObject/WeakObjectPtr.h"
+#include "Framework/Docking/TabManager.h"
+
+class IDetailsView;
+class UActorSequence;
+class UActorSequenceComponent;
+class ISequencer;
+class FSCSEditorTreeNode;
+class IPropertyUtilities;
+
+class FActorSequenceComponentCustomization : public IDetailCustomization
+{
+public:
+	static TSharedRef<IDetailCustomization> MakeInstance();
+
+	virtual void CustomizeDetails(IDetailLayoutBuilder& DetailBuilder) override;
+
+private:
+
+	FReply InvokeSequencer();
+	UActorSequence* GetActorSequence() const;
+
+	TWeakObjectPtr<UActorSequenceComponent> WeakSequenceComponent;
+	TWeakPtr<FTabManager> WeakTabManager;
+	TSharedPtr<SBox> InlineSequencer;
+	TSharedPtr<IPropertyUtilities> PropertyUtilities;
+};
\ No newline at end of file
diff --git a/Engine/Plugins/MovieScene/ActorSequence/Source/ActorSequenceEditor/Private/ActorSequenceEditor.cpp b/Engine/Plugins/MovieScene/ActorSequence/Source/ActorSequenceEditor/Private/ActorSequenceEditor.cpp
new file mode 100644
--- /dev/null
+++ b/Engine/Plugins/MovieScene/ActorSequence/Source/ActorSequenceEditor/Private/ActorSequenceEditor.cpp
@@ -0,0 +1,219 @@
+// Copyright 1998-2018 Epic Games, Inc. All Rights Reserved.
+
+#include "ActorSequence.h"
+#include "ActorSequenceComponent.h"
+#include "BlueprintEditorModule.h"
+#include "BlueprintEditorTabs.h"
+#include "ActorSequenceComponentCustomization.h"
+#include "MovieSceneSequenceEditor_ActorSequence.h"
+#include "ActorSequenceEditorStyle.h"
+#include "ActorSequenceEditorTabSummoner.h"
+#include "Framework/Docking/LayoutExtender.h"
+#include "LevelEditor.h"
+#include "MovieSceneToolsProjectSettings.h"
+#include "PropertyEditorModule.h"
+#include "Styling/SlateStyle.h"
+#include "WorkflowOrientedApp/WorkflowTabManager.h"
+#include "Modules/ModuleManager.h"
+#include "Widgets/Docking/SDockTab.h"
+#include "ISettingsModule.h"
+#include "SequencerSettings.h"
+#include "ISequencerModule.h"
+
+
+#define LOCTEXT_NAMESPACE "ActorSequenceEditor"
+
+
+/** Shared class type that ensures safe binding to RegisterBlueprintEditorTab through an SP binding without interfering with module ownership semantics */
+class FActorSequenceEditorTabBinding
+	: public TSharedFromThis<FActorSequenceEditorTabBinding>
+{
+public:
+
+	FActorSequenceEditorTabBinding()
+	{
+		FBlueprintEditorModule& BlueprintEditorModule = FModuleManager::LoadModuleChecked<FBlueprintEditorModule>("Kismet");
+		BlueprintEditorTabSpawnerHandle = BlueprintEditorModule.OnRegisterTabsForEditor().AddRaw(this, &FActorSequenceEditorTabBinding::RegisterBlueprintEditorTab);
+		BlueprintEditorLayoutExtensionHandle = BlueprintEditorModule.OnRegisterLayoutExtensions().AddRaw(this, &FActorSequenceEditorTabBinding::RegisterBlueprintEditorLayout);
+
+		FLevelEditorModule& LevelEditor = FModuleManager::LoadModuleChecked<FLevelEditorModule>("LevelEditor");
+		LevelEditorTabSpawnerHandle = LevelEditor.OnRegisterTabs().AddRaw(this, &FActorSequenceEditorTabBinding::RegisterLevelEditorTab);
+		LevelEditorLayoutExtensionHandle = LevelEditor.OnRegisterLayoutExtensions().AddRaw(this, &FActorSequenceEditorTabBinding::RegisterLevelEditorLayout);
+	}
+
+	void RegisterLevelEditorLayout(FLayoutExtender& Extender)
+	{
+		Extender.ExtendLayout(FTabId("ContentBrowserTab1"), ELayoutExtensionPosition::Before, FTabManager::FTab(FName("EmbeddedSequenceID"), ETabState::ClosedTab));
+	}
+
+	void RegisterBlueprintEditorLayout(FLayoutExtender& Extender)
+	{
+		Extender.ExtendLayout(FBlueprintEditorTabs::CompilerResultsID, ELayoutExtensionPosition::Before, FTabManager::FTab(FName("EmbeddedSequenceID"), ETabState::ClosedTab));
+	}
+
+	void RegisterBlueprintEditorTab(FWorkflowAllowedTabSet& TabFactories, FName InModeName, TSharedPtr<FBlueprintEditor> BlueprintEditor)
+	{
+		TabFactories.RegisterFactory(MakeShared<FActorSequenceEditorSummoner>(BlueprintEditor));
+	}
+
+	void RegisterLevelEditorTab(TSharedPtr<FTabManager> InTabManager)
+	{
+		auto SpawnTab = [](const FSpawnTabArgs&) -> TSharedRef<SDockTab>
+			{
+				TSharedRef<SActorSequenceEditorWidget> Widget = SNew(SActorSequenceEditorWidget, nullptr);
+
+				return SNew(SDockTab)
+					.Label(&Widget.Get(), &SActorSequenceEditorWidget::GetDisplayLabel)
+					.Icon(FActorSequenceEditorStyle::Get().GetBrush("ClassIcon.ActorSequence"))
+					[
+						Widget
+					];
+			};
+
+		InTabManager->RegisterTabSpawner("EmbeddedSequenceID", FOnSpawnTab::CreateStatic(SpawnTab))
+			.SetMenuType(ETabSpawnerMenuType::Hidden)
+			.SetAutoGenerateMenuEntry(false);
+	}
+
+	~FActorSequenceEditorTabBinding()
+	{
+		FBlueprintEditorModule* BlueprintEditorModule = FModuleManager::GetModulePtr<FBlueprintEditorModule>("Kismet");
+		if (BlueprintEditorModule)
+		{
+			BlueprintEditorModule->OnRegisterTabsForEditor().Remove(BlueprintEditorTabSpawnerHandle);
+			BlueprintEditorModule->OnRegisterLayoutExtensions().Remove(BlueprintEditorLayoutExtensionHandle);
+		}
+
+		FLevelEditorModule* LevelEditor = FModuleManager::GetModulePtr<FLevelEditorModule>("LevelEditor");
+		if (LevelEditor)
+		{
+			LevelEditor->OnRegisterTabs().Remove(LevelEditorTabSpawnerHandle);
+			LevelEditor->OnRegisterLayoutExtensions().Remove(LevelEditorLayoutExtensionHandle);
+		}
+	}
+
+private:
+
+	/** Delegate binding handle for FBlueprintEditorModule::OnRegisterTabsForEditor */
+	FDelegateHandle BlueprintEditorTabSpawnerHandle, BlueprintEditorLayoutExtensionHandle;
+
+	FDelegateHandle LevelEditorTabSpawnerHandle, LevelEditorLayoutExtensionHandle;
+};
+
+/**
+ * Implements the ActorSequenceEditor module.
+ */
+class FActorSequenceEditorModule
+	: public IModuleInterface, public FGCObject
+{
+public:
+
+	FActorSequenceEditorModule()
+		: Settings(nullptr)
+	{
+	}
+
+	virtual void StartupModule() override
+	{
+		// Register styles
+		FActorSequenceEditorStyle::Get();
+
+		BlueprintEditorTabBinding = MakeShared<FActorSequenceEditorTabBinding>();
+		RegisterCustomizations();
+		RegisterSettings();
+		OnInitializeSequenceHandle = UActorSequence::OnInitializeSequence().AddStatic(FActorSequenceEditorModule::OnInitializeSequence);
+
+		ISequencerModule& SequencerModule = FModuleManager::Get().LoadModuleChecked<ISequencerModule>("Sequencer");
+		SequenceEditorHandle = SequencerModule.RegisterSequenceEditor(UActorSequence::StaticClass(), MakeUnique<FMovieSceneSequenceEditor_ActorSequence>());
+	}
+	
+	virtual void ShutdownModule() override
+	{
+		UActorSequence::OnInitializeSequence().Remove(OnInitializeSequenceHandle);
+		UnregisterCustomizations();
+		UnregisterSettings();
+
+		ISequencerModule* SequencerModule = FModuleManager::Get().GetModulePtr<ISequencerModule>("Sequencer");
+		if (SequencerModule)
+		{
+			SequencerModule->UnregisterSequenceEditor(SequenceEditorHandle);
+		}
+
+		BlueprintEditorTabBinding = nullptr;
+	}
+
+	static void OnInitializeSequence(UActorSequence* Sequence)
+	{
+		auto* ProjectSettings = GetDefault<UMovieSceneToolsProjectSettings>();
+		UMovieScene* MovieScene = Sequence->GetMovieScene();
+		
+		FFrameNumber StartFrame = (ProjectSettings->DefaultStartTime * MovieScene->GetTickResolution()).RoundToFrame();
+		int32        Duration   = (ProjectSettings->DefaultDuration * MovieScene->GetTickResolution()).RoundToFrame().Value;
+
+		MovieScene->SetPlaybackRange(StartFrame, Duration);
+	}
+
+	/** Register details view customizations. */
+	void RegisterCustomizations()
+	{
+		FPropertyEditorModule& PropertyModule = FModuleManager::LoadModuleChecked<FPropertyEditorModule>("PropertyEditor");
+		ActorSequenceComponentName = UActorSequenceComponent::StaticClass()->GetFName();
+		PropertyModule.RegisterCustomClassLayout("ActorSequenceComponent", FOnGetDetailCustomizationInstance::CreateStatic(&FActorSequenceComponentCustomization::MakeInstance));
+	}
+
+	/** Unregister details view customizations. */
+	void UnregisterCustomizations()
+	{
+		FPropertyEditorModule* PropertyModule = FModuleManager::GetModulePtr<FPropertyEditorModule>("PropertyEditor");
+		if (PropertyModule)
+		{
+			PropertyModule->UnregisterCustomPropertyTypeLayout(ActorSequenceComponentName);
+		}
+	}
+
+	/** Register settings objects. */
+	void RegisterSettings()
+	{
+		ISettingsModule* SettingsModule = FModuleManager::GetModulePtr<ISettingsModule>("Settings");
+
+		if (SettingsModule != nullptr)
+		{
+			Settings = USequencerSettingsContainer::GetOrCreate<USequencerSettings>(TEXT("EmbeddedActorSequenceEditor"));
+
+			SettingsModule->RegisterSettings("Editor", "ContentEditors", "EmbeddedActorSequenceEditor",
+				LOCTEXT("EmbeddedActorSequenceEditorSettingsName", "Embedded Actor Sequence Editor"),
+				LOCTEXT("EmbeddedActorSequenceEditorSettingsDescription", "Configure the look and feel of the Embedded Actor Sequence Editor."),
+				Settings);	
+		}
+	}
+
+	/** Unregister settings objects. */
+	void UnregisterSettings()
+	{
+		ISettingsModule* SettingsModule = FModuleManager::GetModulePtr<ISettingsModule>("Settings");
+
+		if (SettingsModule != nullptr)
+		{
+			SettingsModule->UnregisterSettings("Editor", "ContentEditors", "EmbeddedActorSequenceEditor");
+		}
+	}
+
+	/** FGCObject interface */
+	virtual void AddReferencedObjects( FReferenceCollector& Collector ) override
+	{
+		if (Settings)
+		{
+			Collector.AddReferencedObject(Settings);
+		}
+	}
+
+	FDelegateHandle SequenceEditorHandle;
+	FDelegateHandle OnInitializeSequenceHandle;
+	TSharedPtr<FActorSequenceEditorTabBinding> BlueprintEditorTabBinding;
+	FName ActorSequenceComponentName;
+	USequencerSettings* Settings;
+};
+
+IMPLEMENT_MODULE(FActorSequenceEditorModule, ActorSequenceEditor);
+
+#undef LOCTEXT_NAMESPACE
diff --git a/Engine/Plugins/MovieScene/ActorSequence/Source/ActorSequenceEditor/Private/ActorSequenceEditorStyle.h b/Engine/Plugins/MovieScene/ActorSequence/Source/ActorSequenceEditor/Private/ActorSequenceEditorStyle.h
new file mode 100644
--- /dev/null
+++ b/Engine/Plugins/MovieScene/ActorSequence/Source/ActorSequenceEditor/Private/ActorSequenceEditorStyle.h
@@ -0,0 +1,35 @@
+// Copyright 1998-2018 Epic Games, Inc. All Rights Reserved.
+
+#pragma once
+
+#include "Styling/SlateStyle.h"
+#include "Styling/SlateStyleRegistry.h"
+
+
+class FActorSequenceEditorStyle
+	: public FSlateStyleSet
+{
+public:
+	FActorSequenceEditorStyle()
+		: FSlateStyleSet("ActorSequenceEditorStyle")
+	{
+		const FVector2D Icon16x16(16.0f, 16.0f);
+		SetContentRoot(FPaths::EnginePluginsDir() / TEXT("MovieScene/ActorSequence/Content"));
+
+		Set("ClassIcon.ActorSequence", new FSlateImageBrush(RootToContentDir(TEXT("ActorSequence_16x.png")), Icon16x16));
+		Set("ClassIcon.ActorSequenceComponent", new FSlateImageBrush(RootToContentDir(TEXT("ActorSequence_16x.png")), Icon16x16));
+
+		FSlateStyleRegistry::RegisterSlateStyle(*this);
+	}
+
+	static FActorSequenceEditorStyle& Get()
+	{
+		static FActorSequenceEditorStyle Inst;
+		return Inst;
+	}
+	
+	~FActorSequenceEditorStyle()
+	{
+		FSlateStyleRegistry::UnRegisterSlateStyle(*this);
+	}
+};
diff --git a/Engine/Plugins/MovieScene/ActorSequence/Source/ActorSequenceEditor/Private/ActorSequenceEditorTabSummoner.cpp b/Engine/Plugins/MovieScene/ActorSequence/Source/ActorSequenceEditor/Private/ActorSequenceEditorTabSummoner.cpp
new file mode 100644
--- /dev/null
+++ b/Engine/Plugins/MovieScene/ActorSequence/Source/ActorSequenceEditor/Private/ActorSequenceEditorTabSummoner.cpp
@@ -0,0 +1,531 @@
+// Copyright 1998-2018 Epic Games, Inc. All Rights Reserved.
+
+#include "ActorSequenceEditorTabSummoner.h"
+
+#include "ActorSequence.h"
+#include "ISequencer.h"
+#include "ISequencerModule.h"
+#include "LevelEditorSequencerIntegration.h"
+#include "SSCSEditor.h"
+#include "Styling/SlateIconFinder.h"
+#include "Kismet2/BlueprintEditorUtils.h"
+#include "EditorStyleSet.h"
+#include "Widgets/Images/SImage.h"
+#include "Editor.h"
+#include "ScopedTransaction.h"
+#include "Framework/MultiBox/MultiBoxBuilder.h"
+#include "Framework/Application/SlateApplication.h"
+
+
+#define LOCTEXT_NAMESPACE "ActorSequenceEditorSummoner"
+
+DECLARE_DELEGATE_OneParam(FOnComponentSelected, TSharedPtr<FSCSEditorTreeNode>);
+DECLARE_DELEGATE_RetVal_OneParam(bool, FIsComponentValid, UActorComponent*);
+
+
+class SComponentSelectionTree
+	: public SCompoundWidget
+{
+public:
+	SLATE_BEGIN_ARGS(SComponentSelectionTree) : _IsInEditMode(false) {}
+
+		SLATE_EVENT(FOnComponentSelected, OnComponentSelected)
+		SLATE_EVENT(FIsComponentValid, IsComponentValid)
+		SLATE_ARGUMENT(bool, IsInEditMode)
+
+	SLATE_END_ARGS()
+
+	void Construct(const FArguments& InArgs, AActor* InPreviewActor)
+	{
+		bIsInEditMode = InArgs._IsInEditMode;
+		OnComponentSelected = InArgs._OnComponentSelected;
+		IsComponentValid = InArgs._IsComponentValid;
+
+		ChildSlot
+		[
+			SAssignNew(TreeView, STreeView<TSharedPtr<FSCSEditorTreeNode>>)
+			.TreeItemsSource(&RootNodes)
+			.SelectionMode(ESelectionMode::Single)
+			.OnGenerateRow(this, &SComponentSelectionTree::GenerateRow)
+			.OnGetChildren(this, &SComponentSelectionTree::OnGetChildNodes)
+			.OnSelectionChanged(this, &SComponentSelectionTree::OnSelectionChanged)
+			.ItemHeight(24)
+		];
+
+		BuildTree(InPreviewActor);
+
+		if (RootNodes.Num() == 0)
+		{
+			ChildSlot
+			[
+				SNew(SBox)
+				.Padding(FMargin(5.f))
+				[
+					SNew(STextBlock)
+					.Text(LOCTEXT("NoValidComponentsFound", "No valid components available"))
+				]
+			];
+		}
+	}
+
+	void BuildTree(AActor* Actor)
+	{
+		RootNodes.Reset();
+		ObjectToNode.Reset();
+
+		for (UActorComponent* Component : TInlineComponentArray<UActorComponent*>(Actor))
+		{
+			if (IsComponentVisibleInTree(Component))
+			{
+				FindOrAddNodeForComponent(Component);
+			}
+		}
+	}
+
+private:
+
+	void OnSelectionChanged(TSharedPtr<FSCSEditorTreeNode> InNode, ESelectInfo::Type SelectInfo)
+	{
+		OnComponentSelected.ExecuteIfBound(InNode);
+	}
+
+	void OnGetChildNodes(TSharedPtr<FSCSEditorTreeNode> InNodePtr, TArray<TSharedPtr<FSCSEditorTreeNode>>& OutChildren)
+	{
+		OutChildren = InNodePtr->GetChildren();
+	}
+
+	TSharedRef<ITableRow> GenerateRow(TSharedPtr<FSCSEditorTreeNode> InNodePtr, const TSharedRef<STableViewBase>& OwnerTable)
+	{
+		const FSlateBrush* ComponentIcon = FEditorStyle::GetBrush("SCS.NativeComponent");
+		if (InNodePtr->GetComponentTemplate() != NULL)
+		{
+			ComponentIcon = FSlateIconFinder::FindIconBrushForClass( InNodePtr->GetComponentTemplate()->GetClass(), TEXT("SCS.Component") );
+		}
+
+		FText Label = InNodePtr->IsInherited() && !bIsInEditMode
+			? FText::Format(LOCTEXT("NativeComponentFormatString","{0} (Inherited)"), FText::FromString(InNodePtr->GetDisplayString()))
+			: FText::FromString(InNodePtr->GetDisplayString());
+
+		TSharedRef<STableRow<FSCSEditorTreeNodePtrType>> Row = SNew(STableRow<FSCSEditorTreeNodePtrType>, OwnerTable).Padding(FMargin(0.f, 0.f, 0.f, 4.f));
+		Row->SetContent(
+			SNew(SHorizontalBox)
+
+			+SHorizontalBox::Slot()
+			.AutoWidth()
+			.VAlign(VAlign_Center)
+			[
+				SNew(SImage)
+				.Image(ComponentIcon)
+				.ColorAndOpacity(SSCS_RowWidget::GetColorTintForIcon(InNodePtr))
+			]
+
+			+SHorizontalBox::Slot()
+			.VAlign(VAlign_Center)
+			.Padding(2, 0, 0, 0)
+			[
+				SNew(STextBlock)
+				.Text(Label)
+			]);
+
+		return Row;
+	}
+
+	bool IsComponentVisibleInTree(UActorComponent* ActorComponent) const
+	{
+		return !IsComponentValid.IsBound() || IsComponentValid.Execute(ActorComponent);
+	}
+
+	TSharedPtr<FSCSEditorTreeNode> FindOrAddNodeForComponent(UActorComponent* ActorComponent)
+	{
+		if (ActorComponent->IsEditorOnly())
+		{
+			return nullptr;
+		}
+
+		if (TSharedPtr<FSCSEditorTreeNode>* Existing = ObjectToNode.Find(ActorComponent))
+		{
+			return *Existing;
+		}
+		else if (USceneComponent* SceneComponent = Cast<USceneComponent>(ActorComponent))
+		{
+			if (UActorComponent* Parent = SceneComponent->GetAttachParent())
+			{
+				TSharedPtr<FSCSEditorTreeNode> ParentNode = FindOrAddNodeForComponent(Parent);
+
+				if (!ParentNode.IsValid())
+				{
+					return nullptr;
+				}
+
+				TreeView->SetItemExpansion(ParentNode, true);
+
+				TSharedPtr<FSCSEditorTreeNode> ChildNode = ParentNode->AddChildFromComponent(ActorComponent);
+				ObjectToNode.Add(ActorComponent, ChildNode);
+
+				return ChildNode;
+			}
+		}
+		
+		TSharedPtr<FSCSEditorTreeNode> RootNode = FSCSEditorTreeNode::FactoryNodeFromComponent(ActorComponent);
+		RootNodes.Add(RootNode);
+		ObjectToNode.Add(ActorComponent, RootNode);
+
+		TreeView->SetItemExpansion(RootNode, true);
+
+		return RootNode;
+	}
+
+private:
+	bool bIsInEditMode;
+	FOnComponentSelected OnComponentSelected;
+	FIsComponentValid IsComponentValid;
+	TSharedPtr<STreeView<TSharedPtr<FSCSEditorTreeNode>>> TreeView;
+	TMap<FObjectKey, TSharedPtr<FSCSEditorTreeNode>> ObjectToNode;
+	TArray<TSharedPtr<FSCSEditorTreeNode>> RootNodes;
+};
+
+class SActorSequenceEditorWidgetImpl : public SCompoundWidget
+{
+public:
+
+	SLATE_BEGIN_ARGS(SActorSequenceEditorWidgetImpl){}
+	SLATE_END_ARGS();
+
+	~SActorSequenceEditorWidgetImpl()
+	{
+		if (Sequencer.IsValid())
+		{
+			FLevelEditorSequencerIntegration::Get().RemoveSequencer(Sequencer.ToSharedRef());
+			Sequencer->Close();
+			Sequencer = nullptr;
+		}
+
+		GEditor->OnBlueprintPreCompile().Remove(OnBlueprintPreCompileHandle);
+		FCoreUObjectDelegates::OnObjectSaved.Remove(OnObjectSavedHandle);
+	}
+
+	void Construct(const FArguments&, TWeakPtr<FBlueprintEditor> InBlueprintEditor)
+	{
+		OnBlueprintPreCompileHandle = GEditor->OnBlueprintPreCompile().AddSP(this, &SActorSequenceEditorWidgetImpl::OnBlueprintPreCompile);
+		OnObjectSavedHandle = FCoreUObjectDelegates::OnObjectSaved.AddSP(this, &SActorSequenceEditorWidgetImpl::OnObjectPreSave);
+
+		WeakBlueprintEditor = InBlueprintEditor;
+
+		ChildSlot
+		[
+			SAssignNew(Content, SBox)
+			.MinDesiredHeight(200)
+		];
+	}
+
+	FText GetDisplayLabel() const
+	{
+		UActorSequence* Sequence = WeakSequence.Get();
+		return Sequence ? Sequence->GetDisplayName() : LOCTEXT("DefaultSequencerLabel", "Sequencer");
+	}
+
+	UActorSequence* GetActorSequence() const
+	{
+		return WeakSequence.Get();
+	}
+
+	UObject* GetPlaybackContext() const
+	{
+		UActorSequence* LocalActorSequence = GetActorSequence();
+		if (LocalActorSequence)
+		{
+			if (AActor* Actor = LocalActorSequence->GetTypedOuter<AActor>())
+			{
+				return Actor;
+			}
+			else if (UBlueprintGeneratedClass* GeneratedClass = LocalActorSequence->GetTypedOuter<UBlueprintGeneratedClass>())
+			{
+				if (GeneratedClass->SimpleConstructionScript)
+				{
+					return GeneratedClass->SimpleConstructionScript->GetComponentEditorActorInstance();
+				}
+			}
+		}
+		
+		return nullptr;
+	}
+
+	TArray<UObject*> GetEventContexts() const
+	{
+		TArray<UObject*> Contexts;
+		if (auto* Context = GetPlaybackContext())
+		{
+			Contexts.Add(Context);
+		}
+		return Contexts;
+	}
+
+	void SetActorSequence(UActorSequence* NewSequence)
+	{
+		if (UActorSequence* OldSequence = WeakSequence.Get())
+		{
+			if (OnSequenceChangedHandle.IsValid())
+			{
+				OldSequence->OnSignatureChanged().Remove(OnSequenceChangedHandle);
+			}
+		}
+
+		WeakSequence = NewSequence;
+
+		if (NewSequence)
+		{
+			OnSequenceChangedHandle = NewSequence->OnSignatureChanged().AddSP(this, &SActorSequenceEditorWidgetImpl::OnSequenceChanged);
+		}
+
+		// If we already have a sequencer open, just assign the sequence
+		if (Sequencer.IsValid() && NewSequence)
+		{
+			if (Sequencer->GetRootMovieSceneSequence() != NewSequence)
+			{
+				Sequencer->ResetToNewRootSequence(*NewSequence);
+			}
+			return;
+		}
+
+		// If we're setting the sequence to none, destroy sequencer
+		if (!NewSequence)
+		{
+			if (Sequencer.IsValid())
+			{
+				FLevelEditorSequencerIntegration::Get().RemoveSequencer(Sequencer.ToSharedRef());
+				Sequencer->Close();
+				Sequencer = nullptr;
+			}
+
+			Content->SetContent(SNew(STextBlock).Text(LOCTEXT("NothingSelected", "Select a sequence")));
+			return;
+		}
+
+		// We need to initialize a new sequencer instance
+		FSequencerInitParams SequencerInitParams;
+		{
+			TWeakObjectPtr<UActorSequence> LocalWeakSequence = NewSequence;
+
+			SequencerInitParams.RootSequence = NewSequence;
+			SequencerInitParams.EventContexts = TAttribute<TArray<UObject*>>(this, &SActorSequenceEditorWidgetImpl::GetEventContexts);
+			SequencerInitParams.PlaybackContext = TAttribute<UObject*>(this, &SActorSequenceEditorWidgetImpl::GetPlaybackContext);
+
+			TSharedRef<FExtender> AddMenuExtender = MakeShareable(new FExtender);
+
+			AddMenuExtender->AddMenuExtension("AddTracks", EExtensionHook::Before, nullptr,
+				FMenuExtensionDelegate::CreateLambda([=](FMenuBuilder& MenuBuilder){
+
+					MenuBuilder.AddSubMenu(
+						LOCTEXT("AddComponent_Label", "Component"),
+						LOCTEXT("AddComponent_ToolTip", "Add a binding to one of this actor's components and allow it to be animated by Sequencer"),
+						FNewMenuDelegate::CreateRaw(this, &SActorSequenceEditorWidgetImpl::AddPossessComponentMenuExtensions),
+						false /*bInOpenSubMenuOnClick*/,
+						FSlateIcon()//"LevelSequenceEditorStyle", "LevelSequenceEditor.PossessNewActor")
+						);
+
+				})
+			);
+
+			SequencerInitParams.ViewParams.bReadOnly = !WeakBlueprintEditor.IsValid() && !NewSequence->IsEditable();
+			SequencerInitParams.bEditWithinLevelEditor = false;
+			SequencerInitParams.ViewParams.AddMenuExtender = AddMenuExtender;
+			SequencerInitParams.ViewParams.UniqueName = "EmbeddedActorSequenceEditor";
+			SequencerInitParams.ViewParams.ScrubberStyle = ESequencerScrubberStyle::FrameBlock;
+			SequencerInitParams.ViewParams.OnReceivedFocus.BindRaw(this, &SActorSequenceEditorWidgetImpl::OnSequencerReceivedFocus);
+		}
+
+		Sequencer = FModuleManager::LoadModuleChecked<ISequencerModule>("Sequencer").CreateSequencer(SequencerInitParams);
+		Content->SetContent(Sequencer->GetSequencerWidget());
+
+		FLevelEditorSequencerIntegrationOptions Options;
+		Options.bRequiresLevelEvents = true;
+		Options.bRequiresActorEvents = false;
+		Options.bCanRecord = false;
+
+		FLevelEditorSequencerIntegration::Get().AddSequencer(Sequencer.ToSharedRef(), Options);
+	}
+
+	void OnSequencerReceivedFocus()
+	{
+		if (Sequencer.IsValid())
+		{
+			FLevelEditorSequencerIntegration::Get().OnSequencerReceivedFocus(Sequencer.ToSharedRef());
+		}
+	}
+
+	void OnObjectPreSave(UObject* InObject)
+	{
+		TSharedPtr<FBlueprintEditor> BlueprintEditor = WeakBlueprintEditor.Pin();
+		if (Sequencer.IsValid() && BlueprintEditor.IsValid() && InObject && InObject == BlueprintEditor->GetBlueprintObj())
+		{
+			Sequencer->RestorePreAnimatedState();
+		}
+	}
+
+	void OnBlueprintPreCompile(UBlueprint* InBlueprint)
+	{
+		TSharedPtr<FBlueprintEditor> BlueprintEditor = WeakBlueprintEditor.Pin();
+		if (Sequencer.IsValid() && BlueprintEditor.IsValid() && InBlueprint && InBlueprint == BlueprintEditor->GetBlueprintObj())
+		{
+			Sequencer->RestorePreAnimatedState();
+		}
+	}
+
+	void OnSelectionUpdated(TSharedPtr<FSCSEditorTreeNode> SelectedNode)
+	{
+		if (SelectedNode->GetNodeType() != FSCSEditorTreeNode::ComponentNode)
+		{
+			return;
+		}
+
+		UActorComponent* EditingComponent = nullptr;
+
+		TSharedPtr<FBlueprintEditor> BlueprintEditor = WeakBlueprintEditor.Pin();
+		if (BlueprintEditor.IsValid())
+		{
+			UBlueprint* Blueprint = BlueprintEditor->GetBlueprintObj();
+			if (Blueprint)
+			{
+				EditingComponent = SelectedNode->GetEditableComponentTemplate(Blueprint);
+			}
+		}
+		else if (AActor* Actor = GetPreviewActor())
+		{
+			EditingComponent = SelectedNode->FindComponentInstanceInActor(Actor);
+		}
+
+		if (EditingComponent)
+		{
+			const FScopedTransaction Transaction(LOCTEXT("AddComponentToSequencer", "Add component to Sequencer"));
+			Sequencer->GetHandleToObject(EditingComponent, true);
+		}
+
+		FSlateApplication::Get().DismissAllMenus();
+	}
+
+	void AddPossessComponentMenuExtensions(FMenuBuilder& MenuBuilder)
+	{
+		AActor* Actor = GetPreviewActor();
+		if (!Actor)
+		{
+			return;
+		}
+
+		Sequencer->State.ClearObjectCaches(*Sequencer);
+		TSet<UObject*> AllBoundObjects;
+
+		AllBoundObjects.Add(GetOwnerComponent());
+
+		UMovieScene* MovieScene = Sequencer->GetFocusedMovieSceneSequence()->GetMovieScene();
+		for (int32 Index = 0; Index < MovieScene->GetPossessableCount(); ++Index)
+		{
+			FMovieScenePossessable& Possessable = MovieScene->GetPossessable(Index);
+			for (TWeakObjectPtr<> WeakObject : Sequencer->FindBoundObjects(Possessable.GetGuid(), Sequencer->GetFocusedTemplateID()))
+			{
+				if (UObject* Object = WeakObject.Get())
+				{
+					AllBoundObjects.Add(Object);
+				}
+			}
+		}
+
+		bool bIdent = false;
+		MenuBuilder.AddWidget(
+			SNew(SComponentSelectionTree, Actor)
+			.IsInEditMode(WeakBlueprintEditor.Pin().IsValid())
+			.OnComponentSelected(this, &SActorSequenceEditorWidgetImpl::OnSelectionUpdated)
+			.IsComponentValid_Lambda(
+				[AllBoundObjects](UActorComponent* Component)
+				{
+					return !AllBoundObjects.Contains(Component);
+				}
+			)
+			, FText(), !bIdent
+		);
+	}
+
+	AActor* GetPreviewActor() const
+	{
+		TSharedPtr<FBlueprintEditor> BlueprintEditor = WeakBlueprintEditor.Pin();
+		if (BlueprintEditor.IsValid())
+		{
+			return BlueprintEditor->GetPreviewActor();
+		}
+		if (UActorSequence* Sequence = WeakSequence.Get())
+		{
+			return Sequence->GetTypedOuter<AActor>();
+		}
+		return nullptr;
+	}
+
+	UActorComponent* GetOwnerComponent() const
+	{
+		UActorSequence* ActorSequence = WeakSequence.Get();
+		AActor* Actor = ActorSequence ? GetPreviewActor() : nullptr;
+
+		return Actor ? FindObject<UActorComponent>(Actor, *ActorSequence->GetOuter()->GetName()) : nullptr;
+	}
+
+	void OnSequenceChanged()
+	{
+		UActorSequence* ActorSequence = WeakSequence.Get();
+		UBlueprint* Blueprint = ActorSequence ? ActorSequence->GetParentBlueprint() : nullptr;
+
+		if (Blueprint)
+		{
+			FBlueprintEditorUtils::MarkBlueprintAsModified(Blueprint);
+		}
+	}
+
+private:
+	TWeakObjectPtr<UActorSequence> WeakSequence;
+
+	TWeakPtr<FBlueprintEditor> WeakBlueprintEditor;
+
+	TSharedPtr<SBox> Content;
+	TSharedPtr<ISequencer> Sequencer;
+
+	FDelegateHandle OnBlueprintPreCompileHandle;
+	FDelegateHandle OnObjectSavedHandle;
+
+	FDelegateHandle OnSequenceChangedHandle;
+};
+
+void SActorSequenceEditorWidget::Construct(const FArguments&, TWeakPtr<FBlueprintEditor> InBlueprintEditor)
+{
+	ChildSlot
+	[
+		SAssignNew(Impl, SActorSequenceEditorWidgetImpl, InBlueprintEditor)
+	];
+}
+
+FText SActorSequenceEditorWidget::GetDisplayLabel() const
+{
+	return Impl.Pin()->GetDisplayLabel();
+}
+
+void SActorSequenceEditorWidget::AssignSequence(UActorSequence* NewActorSequence)
+{
+	Impl.Pin()->SetActorSequence(NewActorSequence);
+}
+
+UActorSequence* SActorSequenceEditorWidget::GetSequence() const
+{
+	return Impl.Pin()->GetActorSequence();
+}
+
+FActorSequenceEditorSummoner::FActorSequenceEditorSummoner(TSharedPtr<FBlueprintEditor> BlueprintEditor)
+	: FWorkflowTabFactory("EmbeddedSequenceID", BlueprintEditor)
+	, WeakBlueprintEditor(BlueprintEditor)
+{
+	bIsSingleton = true;
+
+	TabLabel = LOCTEXT("SequencerTabName", "Sequencer");
+}
+
+TSharedRef<SWidget> FActorSequenceEditorSummoner::CreateTabBody(const FWorkflowTabSpawnInfo& Info) const
+{
+	return SNew(SActorSequenceEditorWidget, WeakBlueprintEditor);
+}
+
+#undef LOCTEXT_NAMESPACE
diff --git a/Engine/Plugins/MovieScene/ActorSequence/Source/ActorSequenceEditor/Private/ActorSequenceEditorTabSummoner.h b/Engine/Plugins/MovieScene/ActorSequence/Source/ActorSequenceEditor/Private/ActorSequenceEditorTabSummoner.h
new file mode 100644
--- /dev/null
+++ b/Engine/Plugins/MovieScene/ActorSequence/Source/ActorSequenceEditor/Private/ActorSequenceEditorTabSummoner.h
@@ -0,0 +1,44 @@
+// Copyright 1998-2018 Epic Games, Inc. All Rights Reserved.
+
+#pragma once
+
+#include "WorkflowOrientedApp/WorkflowUObjectDocuments.h"
+
+
+class FBlueprintEditor;
+class SActorSequenceEditorWidgetImpl;
+class UActorSequence;
+class UActorSequenceComponent;
+
+
+class SActorSequenceEditorWidget
+	: public SCompoundWidget
+{
+public:
+
+	SLATE_BEGIN_ARGS(SActorSequenceEditorWidget){}
+	SLATE_END_ARGS();
+
+	void Construct(const FArguments&, TWeakPtr<FBlueprintEditor> InBlueprintEditor);
+	void AssignSequence(UActorSequence* NewActorSequence);
+	UActorSequence* GetSequence() const;
+	FText GetDisplayLabel() const;
+
+private:
+
+	TWeakPtr<SActorSequenceEditorWidgetImpl> Impl;
+};
+
+
+struct FActorSequenceEditorSummoner
+	: public FWorkflowTabFactory
+{
+	FActorSequenceEditorSummoner(TSharedPtr<FBlueprintEditor> BlueprintEditor);
+
+	virtual TSharedRef<SWidget> CreateTabBody(const FWorkflowTabSpawnInfo& Info) const override;
+
+protected:
+
+	TWeakObjectPtr<UActorSequenceComponent> WeakComponent;
+	TWeakPtr<FBlueprintEditor> WeakBlueprintEditor;
+};
diff --git a/Engine/Plugins/MovieScene/ActorSequence/Source/ActorSequenceEditor/Private/MovieSceneSequenceEditor_ActorSequence.h b/Engine/Plugins/MovieScene/ActorSequence/Source/ActorSequenceEditor/Private/MovieSceneSequenceEditor_ActorSequence.h
new file mode 100644
--- /dev/null
+++ b/Engine/Plugins/MovieScene/ActorSequence/Source/ActorSequenceEditor/Private/MovieSceneSequenceEditor_ActorSequence.h
@@ -0,0 +1,65 @@
+// Copyright 1998-2018 Epic Games, Inc. All Rights Reserved.
+
+#pragma once
+
+#include "MovieSceneSequenceEditor.h"
+#include "ActorSequence.h"
+#include "Engine/Level.h"
+#include "Engine/LevelScriptBlueprint.h"
+#include "Kismet2/KismetEditorUtilities.h"
+#include "Tracks/MovieSceneEventTrack.h"
+#include "K2Node_FunctionEntry.h"
+#include "EdGraphSchema_K2.h"
+
+struct FMovieSceneSequenceEditor_ActorSequence : FMovieSceneSequenceEditor
+{
+	virtual bool CanCreateEvents(UMovieSceneSequence* InSequence) const
+	{
+		return true;
+	}
+
+	virtual UBlueprint* GetBlueprintForSequence(UMovieSceneSequence* InSequence) const override
+	{
+		UActorSequence* ActorSequence = CastChecked<UActorSequence>(InSequence);
+		if (UBlueprint* Blueprint = ActorSequence->GetParentBlueprint())
+		{
+			return Blueprint;
+		}
+
+		UActorSequenceComponent* Component = ActorSequence->GetTypedOuter<UActorSequenceComponent>();
+		ULevel* Level = Component ? Component->GetOwner()->GetLevel() : nullptr;
+
+		bool bDontCreateNewBlueprint = true;
+		return Level ? Level->GetLevelScriptBlueprint(bDontCreateNewBlueprint) : nullptr;
+	}
+
+	virtual UBlueprint* CreateBlueprintForSequence(UMovieSceneSequence* InSequence) const override
+	{
+		UActorSequence* ActorSequence = CastChecked<UActorSequence>(InSequence);
+		check(!ActorSequence->GetParentBlueprint());
+
+		UActorSequenceComponent* Component = ActorSequence->GetTypedOuter<UActorSequenceComponent>();
+		ULevel* Level = Component ? Component->GetOwner()->GetLevel() : nullptr;
+
+		bool bDontCreateNewBlueprint = false;
+		return Level ? Level->GetLevelScriptBlueprint(bDontCreateNewBlueprint) : nullptr;
+	}
+
+	virtual void SetupDefaultPinForEndpoint(UMovieSceneEventTrack* EventTrack, UK2Node_FunctionEntry* Endpoint) const override
+	{
+		// By default for master tracks we create a pin (that receives the level blueprint for master event tracks, or the event receivers)
+		UClass* PinClass = nullptr;
+
+		// When event receivers are not set, we set the pin type to the type of the object binding
+		if (EventTrack->EventReceivers.Num() == 0)
+		{
+			PinClass = FindTrackObjectBindingClass(EventTrack);
+		}
+
+		FEdGraphPinType PinType;
+		PinType.PinCategory = UEdGraphSchema_K2::PC_Object;
+		PinType.PinSubCategoryObject = PinClass ? PinClass : UObject::StaticClass();
+
+		Endpoint->CreateUserDefinedPin(TargetPinName, PinType, EGPD_Output, true);
+	}
+};
\ No newline at end of file
diff --git a/Engine/Plugins/MovieScene/ActorSequenceEditor/ActorSequenceEditor.uplugin b/Engine/Plugins/MovieScene/ActorSequenceEditor/ActorSequenceEditor.uplugin
deleted file mode 100644
diff --git a/Engine/Plugins/MovieScene/ActorSequenceEditor/Source/ActorSequenceEditor/ActorSequenceEditor.Build.cs b/Engine/Plugins/MovieScene/ActorSequenceEditor/Source/ActorSequenceEditor/ActorSequenceEditor.Build.cs
deleted file mode 100644
diff --git a/Engine/Plugins/MovieScene/ActorSequenceEditor/Source/ActorSequenceEditor/Private/ActorSequenceComponentCustomization.cpp b/Engine/Plugins/MovieScene/ActorSequenceEditor/Source/ActorSequenceEditor/Private/ActorSequenceComponentCustomization.cpp
deleted file mode 100644
diff --git a/Engine/Plugins/MovieScene/ActorSequenceEditor/Source/ActorSequenceEditor/Private/ActorSequenceComponentCustomization.h b/Engine/Plugins/MovieScene/ActorSequenceEditor/Source/ActorSequenceEditor/Private/ActorSequenceComponentCustomization.h
deleted file mode 100644
diff --git a/Engine/Plugins/MovieScene/ActorSequenceEditor/Source/ActorSequenceEditor/Private/ActorSequenceEditor.cpp b/Engine/Plugins/MovieScene/ActorSequenceEditor/Source/ActorSequenceEditor/Private/ActorSequenceEditor.cpp
deleted file mode 100644
diff --git a/Engine/Plugins/MovieScene/ActorSequenceEditor/Source/ActorSequenceEditor/Private/ActorSequenceEditorStyle.h b/Engine/Plugins/MovieScene/ActorSequenceEditor/Source/ActorSequenceEditor/Private/ActorSequenceEditorStyle.h
deleted file mode 100644
diff --git a/Engine/Plugins/MovieScene/ActorSequenceEditor/Source/ActorSequenceEditor/Private/ActorSequenceEditorTabSummoner.cpp b/Engine/Plugins/MovieScene/ActorSequenceEditor/Source/ActorSequenceEditor/Private/ActorSequenceEditorTabSummoner.cpp
deleted file mode 100644
diff --git a/Engine/Plugins/MovieScene/ActorSequenceEditor/Source/ActorSequenceEditor/Private/ActorSequenceEditorTabSummoner.h b/Engine/Plugins/MovieScene/ActorSequenceEditor/Source/ActorSequenceEditor/Private/ActorSequenceEditorTabSummoner.h
deleted file mode 100644
diff --git a/Engine/Plugins/MovieScene/ActorSequenceEditor/Source/ActorSequenceEditor/Private/MovieSceneSequenceEditor_ActorSequence.h b/Engine/Plugins/MovieScene/ActorSequenceEditor/Source/ActorSequenceEditor/Private/MovieSceneSequenceEditor_ActorSequence.h
deleted file mode 100644
