Copying //UE4/Dev-Sequencer to Dev-Main (//UE4/Dev-Main)

diff --git a/Engine/Plugins/Media/MediaCompositing/Source/MediaCompositing/Private/MovieScene/MovieSceneMediaTemplate.cpp b/Engine/Plugins/Media/MediaCompositing/Source/MediaCompositing/Private/MovieScene/MovieSceneMediaTemplate.cpp
--- a/Engine/Plugins/Media/MediaCompositing/Source/MediaCompositing/Private/MovieScene/MovieSceneMediaTemplate.cpp
+++ b/Engine/Plugins/Media/MediaCompositing/Source/MediaCompositing/Private/MovieScene/MovieSceneMediaTemplate.cpp
@@ -114,7 +114,7 @@ struct FMediaSectionExecutionToken
 				MediaPlayer->Seek(MediaTime);
 				MediaPlayer->SetRate(1.0f);
 			}
-			else if (Context.HasJumped())
+			else 
 			{
 				MediaPlayer->Seek(MediaTime);
 			}
diff --git a/Engine/Plugins/MovieScene/ActorSequence/ActorSequence.uplugin b/Engine/Plugins/MovieScene/ActorSequence/ActorSequence.uplugin
--- a/Engine/Plugins/MovieScene/ActorSequence/ActorSequence.uplugin
+++ b/Engine/Plugins/MovieScene/ActorSequence/ActorSequence.uplugin
@@ -20,6 +20,11 @@
 			"Name" : "ActorSequence",
 			"Type" : "Runtime",
 			"LoadingPhase" : "PreDefault"
+		},
+		{
+			"Name" : "ActorSequenceEditor",
+			"Type" : "Editor",
+			"LoadingPhase" : "PostEngineInit"
 		}
 	]
 }
\ No newline at end of file
diff --git a/Engine/Plugins/MovieScene/ActorSequence/Source/ActorSequenceEditor/ActorSequenceEditor.Build.cs b/Engine/Plugins/MovieScene/ActorSequence/Source/ActorSequenceEditor/ActorSequenceEditor.Build.cs
new file mode 100644
--- /dev/null
+++ b/Engine/Plugins/MovieScene/ActorSequence/Source/ActorSequenceEditor/ActorSequenceEditor.Build.cs
@@ -0,0 +1,46 @@
+// Copyright 1998-2018 Epic Games, Inc. All Rights Reserved.
+
+using UnrealBuildTool;
+
+public class ActorSequenceEditor : ModuleRules
+{
+	public ActorSequenceEditor(ReadOnlyTargetRules Target) : base(Target)
+	{
+		PrivateDependencyModuleNames.AddRange(
+			new string[] {
+				"ActorSequence",
+				"BlueprintGraph",
+				"Core",
+				"CoreUObject",
+				"Engine",
+				"InputCore",
+				"Kismet",
+				"MovieScene",
+				"MovieSceneTools",
+				"Sequencer",
+				"EditorStyle",
+				"Slate",
+				"SlateCore",
+				"UnrealEd",
+				"TimeManagement"
+			}
+		);
+
+		PrivateIncludePaths.AddRange(
+			new string[] {
+				"ActorSequenceEditor/Private",
+			}
+		);
+
+		var DynamicModuleNames = new string[] {
+			"LevelEditor",
+			"PropertyEditor",
+		};
+
+		foreach (var Name in DynamicModuleNames)
+		{
+			PrivateIncludePathModuleNames.Add(Name);
+			DynamicallyLoadedModuleNames.Add(Name);
+		}
+	}
+}
diff --git a/Engine/Plugins/MovieScene/ActorSequence/Source/ActorSequenceEditor/Private/ActorSequenceComponentCustomization.cpp b/Engine/Plugins/MovieScene/ActorSequence/Source/ActorSequenceEditor/Private/ActorSequenceComponentCustomization.cpp
new file mode 100644
--- /dev/null
+++ b/Engine/Plugins/MovieScene/ActorSequence/Source/ActorSequenceEditor/Private/ActorSequenceComponentCustomization.cpp
@@ -0,0 +1,197 @@
+// Copyright 1998-2018 Epic Games, Inc. All Rights Reserved.
+
+#include "ActorSequenceComponentCustomization.h"
+
+#include "ActorSequence.h"
+#include "ActorSequenceComponent.h"
+#include "EditorStyleSet.h"
+#include "GameFramework/Actor.h"
+#include "IDetailsView.h"
+#include "DetailLayoutBuilder.h"
+#include "DetailCategoryBuilder.h"
+#include "DetailWidgetRow.h"
+#include "Widgets/Docking/SDockTab.h"
+#include "SSCSEditor.h"
+#include "BlueprintEditorTabs.h"
+#include "ScopedTransaction.h"
+#include "ISequencerModule.h"
+#include "Editor.h"
+#include "ActorSequenceEditorTabSummoner.h"
+#include "IPropertyUtilities.h"
+#include "Widgets/Input/SButton.h"
+
+#define LOCTEXT_NAMESPACE "ActorSequenceComponentCustomization"
+
+FName SequenceTabId("EmbeddedSequenceID");
+
+class SActorSequenceEditorWidgetWrapper : public SActorSequenceEditorWidget
+{
+public:
+	~SActorSequenceEditorWidgetWrapper()
+	{
+		GEditor->OnObjectsReplaced().Remove(OnObjectsReplacedHandle);
+	}
+
+	void Construct(const FArguments& InArgs, TWeakObjectPtr<UActorSequenceComponent> InSequenceComponent)
+	{
+		SActorSequenceEditorWidget::Construct(InArgs, nullptr);
+
+		WeakSequenceComponent = InSequenceComponent;
+		AssignSequence(GetActorSequence());
+
+		OnObjectsReplacedHandle = GEditor->OnObjectsReplaced().AddSP(this, &SActorSequenceEditorWidgetWrapper::OnObjectsReplaced);
+	}
+
+protected:
+
+	UActorSequence* GetActorSequence() const
+	{
+		UActorSequenceComponent* SequenceComponent = WeakSequenceComponent.Get();
+		return SequenceComponent ? SequenceComponent->GetSequence() : nullptr;
+	}
+
+	void OnObjectsReplaced(const TMap<UObject*, UObject*>& ReplacementMap)
+	{
+		UActorSequenceComponent* Component = WeakSequenceComponent.Get(true);
+
+		UActorSequenceComponent* NewSequenceComponent = Component ? Cast<UActorSequenceComponent>(ReplacementMap.FindRef(Component)) : nullptr;
+		if (NewSequenceComponent)
+		{
+			WeakSequenceComponent = NewSequenceComponent;
+			AssignSequence(GetActorSequence());
+		}
+	}
+
+private:
+
+	TWeakObjectPtr<UActorSequenceComponent> WeakSequenceComponent;
+	FDelegateHandle OnObjectsReplacedHandle;
+};
+
+TSharedRef<IDetailCustomization> FActorSequenceComponentCustomization::MakeInstance()
+{
+	return MakeShared<FActorSequenceComponentCustomization>();
+}
+
+void FActorSequenceComponentCustomization::CustomizeDetails(IDetailLayoutBuilder& DetailBuilder)
+{
+	PropertyUtilities = DetailBuilder.GetPropertyUtilities();
+
+	TArray<TWeakObjectPtr<UObject>> Objects;
+	DetailBuilder.GetObjectsBeingCustomized(Objects);
+	if (Objects.Num() != 1)
+	{
+		// @todo: Display other common properties / information?
+		return;
+	}
+
+	WeakSequenceComponent = Cast<UActorSequenceComponent>(Objects[0].Get());
+	if (!WeakSequenceComponent.Get())
+	{
+		return;
+	}
+
+	const IDetailsView* DetailsView = DetailBuilder.GetDetailsView();
+	TSharedPtr<FTabManager> HostTabManager = DetailsView->GetHostTabManager();
+
+	DetailBuilder.HideProperty("Sequence");
+
+	IDetailCategoryBuilder& Category = DetailBuilder.EditCategory("Sequence", FText(), ECategoryPriority::Important);
+
+	bool bIsExternalTabAlreadyOpened = false;
+
+	if (HostTabManager.IsValid() && HostTabManager->CanSpawnTab(SequenceTabId))
+	{
+		WeakTabManager = HostTabManager;
+
+		TSharedPtr<SDockTab> ExistingTab = HostTabManager->FindExistingLiveTab(SequenceTabId);
+		if (ExistingTab.IsValid())
+		{
+			UActorSequence* ThisSequence = GetActorSequence();
+
+			auto SequencerWidget = StaticCastSharedRef<SActorSequenceEditorWidget>(ExistingTab->GetContent());
+			bIsExternalTabAlreadyOpened = ThisSequence && SequencerWidget->GetSequence() == ThisSequence;
+		}
+
+		Category.AddCustomRow(FText())
+			.NameContent()
+			[
+				SNew(STextBlock)
+				.Text(LOCTEXT("SequenceValueText", "Sequence"))
+				.Font(DetailBuilder.GetDetailFont())
+			]
+			.ValueContent()
+			[
+				SNew(SButton)
+				.OnClicked(this, &FActorSequenceComponentCustomization::InvokeSequencer)
+				[
+					SNew(STextBlock)
+					.Text(bIsExternalTabAlreadyOpened ? LOCTEXT("FocusSequenceTabButtonText", "Focus Tab") : LOCTEXT("OpenSequenceTabButtonText", "Open in Tab"))
+					.Font(DetailBuilder.GetDetailFont())
+				]
+			];
+	}
+
+	// Only display an inline editor for non-blueprint sequences
+	if (GetActorSequence() && !GetActorSequence()->GetParentBlueprint() && !bIsExternalTabAlreadyOpened)
+	{
+		Category.AddCustomRow(FText())
+		.WholeRowContent()
+		.MaxDesiredWidth(TOptional<float>())
+		[
+			SAssignNew(InlineSequencer, SBox)
+			.HeightOverride(300)
+			[
+				SNew(SActorSequenceEditorWidgetWrapper, WeakSequenceComponent)
+			]
+		];
+	}
+}
+
+FReply FActorSequenceComponentCustomization::InvokeSequencer()
+{
+	TSharedPtr<FTabManager> TabManager = WeakTabManager.Pin();
+	if (TabManager.IsValid() && TabManager->CanSpawnTab(SequenceTabId))
+	{
+		TSharedRef<SDockTab> Tab = TabManager->InvokeTab(SequenceTabId);
+
+		{
+			// Set up a delegate that forces a refresh of this panel when the tab is closed to ensure we see the inline widget
+			TWeakPtr<IPropertyUtilities> WeakUtilities = PropertyUtilities;
+			auto OnClosed = [WeakUtilities](TSharedRef<SDockTab>)
+			{
+				TSharedPtr<IPropertyUtilities> PinnedPropertyUtilities = WeakUtilities.Pin();
+				if (PinnedPropertyUtilities.IsValid())
+				{
+					PinnedPropertyUtilities->EnqueueDeferredAction(FSimpleDelegate::CreateSP(PinnedPropertyUtilities.ToSharedRef(), &IPropertyUtilities::ForceRefresh));
+				}
+			};
+
+			Tab->SetOnTabClosed(SDockTab::FOnTabClosedCallback::CreateLambda(OnClosed));
+		}
+
+		// Move our inline widget content to the tab (so that we keep the existing sequencer state)
+		if (InlineSequencer.IsValid())
+		{
+			Tab->SetContent(InlineSequencer->GetChildren()->GetChildAt(0));
+			InlineSequencer->SetContent(SNullWidget::NullWidget);
+			InlineSequencer->SetVisibility(EVisibility::Collapsed);
+		}
+		else 
+		{
+			StaticCastSharedRef<SActorSequenceEditorWidget>(Tab->GetContent())->AssignSequence(GetActorSequence());
+		}
+	}
+
+	PropertyUtilities->ForceRefresh();
+
+	return FReply::Handled();
+}
+
+UActorSequence* FActorSequenceComponentCustomization::GetActorSequence() const
+{
+	UActorSequenceComponent* SequenceComponent = WeakSequenceComponent.Get();
+	return SequenceComponent ? SequenceComponent->GetSequence() : nullptr;
+}
+
+#undef LOCTEXT_NAMESPACE
diff --git a/Engine/Plugins/MovieScene/ActorSequence/Source/ActorSequenceEditor/Private/ActorSequenceComponentCustomization.h b/Engine/Plugins/MovieScene/ActorSequence/Source/ActorSequenceEditor/Private/ActorSequenceComponentCustomization.h
new file mode 100644
--- /dev/null
+++ b/Engine/Plugins/MovieScene/ActorSequence/Source/ActorSequenceEditor/Private/ActorSequenceComponentCustomization.h
@@ -0,0 +1,35 @@
+// Copyright 1998-2018 Epic Games, Inc. All Rights Reserved.
+
+#pragma once
+
+#include "CoreMinimal.h"
+#include "IDetailCustomization.h"
+#include "Widgets/Layout/SBox.h"
+#include "Input/Reply.h"
+#include "UObject/WeakObjectPtr.h"
+#include "Framework/Docking/TabManager.h"
+
+class IDetailsView;
+class UActorSequence;
+class UActorSequenceComponent;
+class ISequencer;
+class FSCSEditorTreeNode;
+class IPropertyUtilities;
+
+class FActorSequenceComponentCustomization : public IDetailCustomization
+{
+public:
+	static TSharedRef<IDetailCustomization> MakeInstance();
+
+	virtual void CustomizeDetails(IDetailLayoutBuilder& DetailBuilder) override;
+
+private:
+
+	FReply InvokeSequencer();
+	UActorSequence* GetActorSequence() const;
+
+	TWeakObjectPtr<UActorSequenceComponent> WeakSequenceComponent;
+	TWeakPtr<FTabManager> WeakTabManager;
+	TSharedPtr<SBox> InlineSequencer;
+	TSharedPtr<IPropertyUtilities> PropertyUtilities;
+};
\ No newline at end of file
diff --git a/Engine/Plugins/MovieScene/ActorSequence/Source/ActorSequenceEditor/Private/ActorSequenceEditor.cpp b/Engine/Plugins/MovieScene/ActorSequence/Source/ActorSequenceEditor/Private/ActorSequenceEditor.cpp
new file mode 100644
--- /dev/null
+++ b/Engine/Plugins/MovieScene/ActorSequence/Source/ActorSequenceEditor/Private/ActorSequenceEditor.cpp
@@ -0,0 +1,219 @@
+// Copyright 1998-2018 Epic Games, Inc. All Rights Reserved.
+
+#include "ActorSequence.h"
+#include "ActorSequenceComponent.h"
+#include "BlueprintEditorModule.h"
+#include "BlueprintEditorTabs.h"
+#include "ActorSequenceComponentCustomization.h"
+#include "MovieSceneSequenceEditor_ActorSequence.h"
+#include "ActorSequenceEditorStyle.h"
+#include "ActorSequenceEditorTabSummoner.h"
+#include "Framework/Docking/LayoutExtender.h"
+#include "LevelEditor.h"
+#include "MovieSceneToolsProjectSettings.h"
+#include "PropertyEditorModule.h"
+#include "Styling/SlateStyle.h"
+#include "WorkflowOrientedApp/WorkflowTabManager.h"
+#include "Modules/ModuleManager.h"
+#include "Widgets/Docking/SDockTab.h"
+#include "ISettingsModule.h"
+#include "SequencerSettings.h"
+#include "ISequencerModule.h"
+
+
+#define LOCTEXT_NAMESPACE "ActorSequenceEditor"
+
+
+/** Shared class type that ensures safe binding to RegisterBlueprintEditorTab through an SP binding without interfering with module ownership semantics */
+class FActorSequenceEditorTabBinding
+	: public TSharedFromThis<FActorSequenceEditorTabBinding>
+{
+public:
+
+	FActorSequenceEditorTabBinding()
+	{
+		FBlueprintEditorModule& BlueprintEditorModule = FModuleManager::LoadModuleChecked<FBlueprintEditorModule>("Kismet");
+		BlueprintEditorTabSpawnerHandle = BlueprintEditorModule.OnRegisterTabsForEditor().AddRaw(this, &FActorSequenceEditorTabBinding::RegisterBlueprintEditorTab);
+		BlueprintEditorLayoutExtensionHandle = BlueprintEditorModule.OnRegisterLayoutExtensions().AddRaw(this, &FActorSequenceEditorTabBinding::RegisterBlueprintEditorLayout);
+
+		FLevelEditorModule& LevelEditor = FModuleManager::LoadModuleChecked<FLevelEditorModule>("LevelEditor");
+		LevelEditorTabSpawnerHandle = LevelEditor.OnRegisterTabs().AddRaw(this, &FActorSequenceEditorTabBinding::RegisterLevelEditorTab);
+		LevelEditorLayoutExtensionHandle = LevelEditor.OnRegisterLayoutExtensions().AddRaw(this, &FActorSequenceEditorTabBinding::RegisterLevelEditorLayout);
+	}
+
+	void RegisterLevelEditorLayout(FLayoutExtender& Extender)
+	{
+		Extender.ExtendLayout(FTabId("ContentBrowserTab1"), ELayoutExtensionPosition::Before, FTabManager::FTab(FName("EmbeddedSequenceID"), ETabState::ClosedTab));
+	}
+
+	void RegisterBlueprintEditorLayout(FLayoutExtender& Extender)
+	{
+		Extender.ExtendLayout(FBlueprintEditorTabs::CompilerResultsID, ELayoutExtensionPosition::Before, FTabManager::FTab(FName("EmbeddedSequenceID"), ETabState::ClosedTab));
+	}
+
+	void RegisterBlueprintEditorTab(FWorkflowAllowedTabSet& TabFactories, FName InModeName, TSharedPtr<FBlueprintEditor> BlueprintEditor)
+	{
+		TabFactories.RegisterFactory(MakeShared<FActorSequenceEditorSummoner>(BlueprintEditor));
+	}
+
+	void RegisterLevelEditorTab(TSharedPtr<FTabManager> InTabManager)
+	{
+		auto SpawnTab = [](const FSpawnTabArgs&) -> TSharedRef<SDockTab>
+			{
+				TSharedRef<SActorSequenceEditorWidget> Widget = SNew(SActorSequenceEditorWidget, nullptr);
+
+				return SNew(SDockTab)
+					.Label(&Widget.Get(), &SActorSequenceEditorWidget::GetDisplayLabel)
+					.Icon(FActorSequenceEditorStyle::Get().GetBrush("ClassIcon.ActorSequence"))
+					[
+						Widget
+					];
+			};
+
+		InTabManager->RegisterTabSpawner("EmbeddedSequenceID", FOnSpawnTab::CreateStatic(SpawnTab))
+			.SetMenuType(ETabSpawnerMenuType::Hidden)
+			.SetAutoGenerateMenuEntry(false);
+	}
+
+	~FActorSequenceEditorTabBinding()
+	{
+		FBlueprintEditorModule* BlueprintEditorModule = FModuleManager::GetModulePtr<FBlueprintEditorModule>("Kismet");
+		if (BlueprintEditorModule)
+		{
+			BlueprintEditorModule->OnRegisterTabsForEditor().Remove(BlueprintEditorTabSpawnerHandle);
+			BlueprintEditorModule->OnRegisterLayoutExtensions().Remove(BlueprintEditorLayoutExtensionHandle);
+		}
+
+		FLevelEditorModule* LevelEditor = FModuleManager::GetModulePtr<FLevelEditorModule>("LevelEditor");
+		if (LevelEditor)
+		{
+			LevelEditor->OnRegisterTabs().Remove(LevelEditorTabSpawnerHandle);
+			LevelEditor->OnRegisterLayoutExtensions().Remove(LevelEditorLayoutExtensionHandle);
+		}
+	}
+
+private:
+
+	/** Delegate binding handle for FBlueprintEditorModule::OnRegisterTabsForEditor */
+	FDelegateHandle BlueprintEditorTabSpawnerHandle, BlueprintEditorLayoutExtensionHandle;
+
+	FDelegateHandle LevelEditorTabSpawnerHandle, LevelEditorLayoutExtensionHandle;
+};
+
+/**
+ * Implements the ActorSequenceEditor module.
+ */
+class FActorSequenceEditorModule
+	: public IModuleInterface, public FGCObject
+{
+public:
+
+	FActorSequenceEditorModule()
+		: Settings(nullptr)
+	{
+	}
+
+	virtual void StartupModule() override
+	{
+		// Register styles
+		FActorSequenceEditorStyle::Get();
+
+		BlueprintEditorTabBinding = MakeShared<FActorSequenceEditorTabBinding>();
+		RegisterCustomizations();
+		RegisterSettings();
+		OnInitializeSequenceHandle = UActorSequence::OnInitializeSequence().AddStatic(FActorSequenceEditorModule::OnInitializeSequence);
+
+		ISequencerModule& SequencerModule = FModuleManager::Get().LoadModuleChecked<ISequencerModule>("Sequencer");
+		SequenceEditorHandle = SequencerModule.RegisterSequenceEditor(UActorSequence::StaticClass(), MakeUnique<FMovieSceneSequenceEditor_ActorSequence>());
+	}
+	
+	virtual void ShutdownModule() override
+	{
+		UActorSequence::OnInitializeSequence().Remove(OnInitializeSequenceHandle);
+		UnregisterCustomizations();
+		UnregisterSettings();
+
+		ISequencerModule* SequencerModule = FModuleManager::Get().GetModulePtr<ISequencerModule>("Sequencer");
+		if (SequencerModule)
+		{
+			SequencerModule->UnregisterSequenceEditor(SequenceEditorHandle);
+		}
+
+		BlueprintEditorTabBinding = nullptr;
+	}
+
+	static void OnInitializeSequence(UActorSequence* Sequence)
+	{
+		auto* ProjectSettings = GetDefault<UMovieSceneToolsProjectSettings>();
+		UMovieScene* MovieScene = Sequence->GetMovieScene();
+		
+		FFrameNumber StartFrame = (ProjectSettings->DefaultStartTime * MovieScene->GetTickResolution()).RoundToFrame();
+		int32        Duration   = (ProjectSettings->DefaultDuration * MovieScene->GetTickResolution()).RoundToFrame().Value;
+
+		MovieScene->SetPlaybackRange(StartFrame, Duration);
+	}
+
+	/** Register details view customizations. */
+	void RegisterCustomizations()
+	{
+		FPropertyEditorModule& PropertyModule = FModuleManager::LoadModuleChecked<FPropertyEditorModule>("PropertyEditor");
+		ActorSequenceComponentName = UActorSequenceComponent::StaticClass()->GetFName();
+		PropertyModule.RegisterCustomClassLayout("ActorSequenceComponent", FOnGetDetailCustomizationInstance::CreateStatic(&FActorSequenceComponentCustomization::MakeInstance));
+	}
+
+	/** Unregister details view customizations. */
+	void UnregisterCustomizations()
+	{
+		FPropertyEditorModule* PropertyModule = FModuleManager::GetModulePtr<FPropertyEditorModule>("PropertyEditor");
+		if (PropertyModule)
+		{
+			PropertyModule->UnregisterCustomPropertyTypeLayout(ActorSequenceComponentName);
+		}
+	}
+
+	/** Register settings objects. */
+	void RegisterSettings()
+	{
+		ISettingsModule* SettingsModule = FModuleManager::GetModulePtr<ISettingsModule>("Settings");
+
+		if (SettingsModule != nullptr)
+		{
+			Settings = USequencerSettingsContainer::GetOrCreate<USequencerSettings>(TEXT("EmbeddedActorSequenceEditor"));
+
+			SettingsModule->RegisterSettings("Editor", "ContentEditors", "EmbeddedActorSequenceEditor",
+				LOCTEXT("EmbeddedActorSequenceEditorSettingsName", "Embedded Actor Sequence Editor"),
+				LOCTEXT("EmbeddedActorSequenceEditorSettingsDescription", "Configure the look and feel of the Embedded Actor Sequence Editor."),
+				Settings);	
+		}
+	}
+
+	/** Unregister settings objects. */
+	void UnregisterSettings()
+	{
+		ISettingsModule* SettingsModule = FModuleManager::GetModulePtr<ISettingsModule>("Settings");
+
+		if (SettingsModule != nullptr)
+		{
+			SettingsModule->UnregisterSettings("Editor", "ContentEditors", "EmbeddedActorSequenceEditor");
+		}
+	}
+
+	/** FGCObject interface */
+	virtual void AddReferencedObjects( FReferenceCollector& Collector ) override
+	{
+		if (Settings)
+		{
+			Collector.AddReferencedObject(Settings);
+		}
+	}
+
+	FDelegateHandle SequenceEditorHandle;
+	FDelegateHandle OnInitializeSequenceHandle;
+	TSharedPtr<FActorSequenceEditorTabBinding> BlueprintEditorTabBinding;
+	FName ActorSequenceComponentName;
+	USequencerSettings* Settings;
+};
+
+IMPLEMENT_MODULE(FActorSequenceEditorModule, ActorSequenceEditor);
+
+#undef LOCTEXT_NAMESPACE
diff --git a/Engine/Plugins/MovieScene/ActorSequence/Source/ActorSequenceEditor/Private/ActorSequenceEditorStyle.h b/Engine/Plugins/MovieScene/ActorSequence/Source/ActorSequenceEditor/Private/ActorSequenceEditorStyle.h
new file mode 100644
--- /dev/null
+++ b/Engine/Plugins/MovieScene/ActorSequence/Source/ActorSequenceEditor/Private/ActorSequenceEditorStyle.h
@@ -0,0 +1,35 @@
+// Copyright 1998-2018 Epic Games, Inc. All Rights Reserved.
+
+#pragma once
+
+#include "Styling/SlateStyle.h"
+#include "Styling/SlateStyleRegistry.h"
+
+
+class FActorSequenceEditorStyle
+	: public FSlateStyleSet
+{
+public:
+	FActorSequenceEditorStyle()
+		: FSlateStyleSet("ActorSequenceEditorStyle")
+	{
+		const FVector2D Icon16x16(16.0f, 16.0f);
+		SetContentRoot(FPaths::EnginePluginsDir() / TEXT("MovieScene/ActorSequence/Content"));
+
+		Set("ClassIcon.ActorSequence", new FSlateImageBrush(RootToContentDir(TEXT("ActorSequence_16x.png")), Icon16x16));
+		Set("ClassIcon.ActorSequenceComponent", new FSlateImageBrush(RootToContentDir(TEXT("ActorSequence_16x.png")), Icon16x16));
+
+		FSlateStyleRegistry::RegisterSlateStyle(*this);
+	}
+
+	static FActorSequenceEditorStyle& Get()
+	{
+		static FActorSequenceEditorStyle Inst;
+		return Inst;
+	}
+	
+	~FActorSequenceEditorStyle()
+	{
+		FSlateStyleRegistry::UnRegisterSlateStyle(*this);
+	}
+};
diff --git a/Engine/Plugins/MovieScene/ActorSequence/Source/ActorSequenceEditor/Private/ActorSequenceEditorTabSummoner.cpp b/Engine/Plugins/MovieScene/ActorSequence/Source/ActorSequenceEditor/Private/ActorSequenceEditorTabSummoner.cpp
new file mode 100644
--- /dev/null
+++ b/Engine/Plugins/MovieScene/ActorSequence/Source/ActorSequenceEditor/Private/ActorSequenceEditorTabSummoner.cpp
@@ -0,0 +1,531 @@
+// Copyright 1998-2018 Epic Games, Inc. All Rights Reserved.
+
+#include "ActorSequenceEditorTabSummoner.h"
+
+#include "ActorSequence.h"
+#include "ISequencer.h"
+#include "ISequencerModule.h"
+#include "LevelEditorSequencerIntegration.h"
+#include "SSCSEditor.h"
+#include "Styling/SlateIconFinder.h"
+#include "Kismet2/BlueprintEditorUtils.h"
+#include "EditorStyleSet.h"
+#include "Widgets/Images/SImage.h"
+#include "Editor.h"
+#include "ScopedTransaction.h"
+#include "Framework/MultiBox/MultiBoxBuilder.h"
+#include "Framework/Application/SlateApplication.h"
+
+
+#define LOCTEXT_NAMESPACE "ActorSequenceEditorSummoner"
+
+DECLARE_DELEGATE_OneParam(FOnComponentSelected, TSharedPtr<FSCSEditorTreeNode>);
+DECLARE_DELEGATE_RetVal_OneParam(bool, FIsComponentValid, UActorComponent*);
+
+
+class SComponentSelectionTree
+	: public SCompoundWidget
+{
+public:
+	SLATE_BEGIN_ARGS(SComponentSelectionTree) : _IsInEditMode(false) {}
+
+		SLATE_EVENT(FOnComponentSelected, OnComponentSelected)
+		SLATE_EVENT(FIsComponentValid, IsComponentValid)
+		SLATE_ARGUMENT(bool, IsInEditMode)
+
+	SLATE_END_ARGS()
+
+	void Construct(const FArguments& InArgs, AActor* InPreviewActor)
+	{
+		bIsInEditMode = InArgs._IsInEditMode;
+		OnComponentSelected = InArgs._OnComponentSelected;
+		IsComponentValid = InArgs._IsComponentValid;
+
+		ChildSlot
+		[
+			SAssignNew(TreeView, STreeView<TSharedPtr<FSCSEditorTreeNode>>)
+			.TreeItemsSource(&RootNodes)
+			.SelectionMode(ESelectionMode::Single)
+			.OnGenerateRow(this, &SComponentSelectionTree::GenerateRow)
+			.OnGetChildren(this, &SComponentSelectionTree::OnGetChildNodes)
+			.OnSelectionChanged(this, &SComponentSelectionTree::OnSelectionChanged)
+			.ItemHeight(24)
+		];
+
+		BuildTree(InPreviewActor);
+
+		if (RootNodes.Num() == 0)
+		{
+			ChildSlot
+			[
+				SNew(SBox)
+				.Padding(FMargin(5.f))
+				[
+					SNew(STextBlock)
+					.Text(LOCTEXT("NoValidComponentsFound", "No valid components available"))
+				]
+			];
+		}
+	}
+
+	void BuildTree(AActor* Actor)
+	{
+		RootNodes.Reset();
+		ObjectToNode.Reset();
+
+		for (UActorComponent* Component : TInlineComponentArray<UActorComponent*>(Actor))
+		{
+			if (IsComponentVisibleInTree(Component))
+			{
+				FindOrAddNodeForComponent(Component);
+			}
+		}
+	}
+
+private:
+
+	void OnSelectionChanged(TSharedPtr<FSCSEditorTreeNode> InNode, ESelectInfo::Type SelectInfo)
+	{
+		OnComponentSelected.ExecuteIfBound(InNode);
+	}
+
+	void OnGetChildNodes(TSharedPtr<FSCSEditorTreeNode> InNodePtr, TArray<TSharedPtr<FSCSEditorTreeNode>>& OutChildren)
+	{
+		OutChildren = InNodePtr->GetChildren();
+	}
+
+	TSharedRef<ITableRow> GenerateRow(TSharedPtr<FSCSEditorTreeNode> InNodePtr, const TSharedRef<STableViewBase>& OwnerTable)
+	{
+		const FSlateBrush* ComponentIcon = FEditorStyle::GetBrush("SCS.NativeComponent");
+		if (InNodePtr->GetComponentTemplate() != NULL)
+		{
+			ComponentIcon = FSlateIconFinder::FindIconBrushForClass( InNodePtr->GetComponentTemplate()->GetClass(), TEXT("SCS.Component") );
+		}
+
+		FText Label = InNodePtr->IsInherited() && !bIsInEditMode
+			? FText::Format(LOCTEXT("NativeComponentFormatString","{0} (Inherited)"), FText::FromString(InNodePtr->GetDisplayString()))
+			: FText::FromString(InNodePtr->GetDisplayString());
+
+		TSharedRef<STableRow<FSCSEditorTreeNodePtrType>> Row = SNew(STableRow<FSCSEditorTreeNodePtrType>, OwnerTable).Padding(FMargin(0.f, 0.f, 0.f, 4.f));
+		Row->SetContent(
+			SNew(SHorizontalBox)
+
+			+SHorizontalBox::Slot()
+			.AutoWidth()
+			.VAlign(VAlign_Center)
+			[
+				SNew(SImage)
+				.Image(ComponentIcon)
+				.ColorAndOpacity(SSCS_RowWidget::GetColorTintForIcon(InNodePtr))
+			]
+
+			+SHorizontalBox::Slot()
+			.VAlign(VAlign_Center)
+			.Padding(2, 0, 0, 0)
+			[
+				SNew(STextBlock)
+				.Text(Label)
+			]);
+
+		return Row;
+	}
+
+	bool IsComponentVisibleInTree(UActorComponent* ActorComponent) const
+	{
+		return !IsComponentValid.IsBound() || IsComponentValid.Execute(ActorComponent);
+	}
+
+	TSharedPtr<FSCSEditorTreeNode> FindOrAddNodeForComponent(UActorComponent* ActorComponent)
+	{
+		if (ActorComponent->IsEditorOnly())
+		{
+			return nullptr;
+		}
+
+		if (TSharedPtr<FSCSEditorTreeNode>* Existing = ObjectToNode.Find(ActorComponent))
+		{
+			return *Existing;
+		}
+		else if (USceneComponent* SceneComponent = Cast<USceneComponent>(ActorComponent))
+		{
+			if (UActorComponent* Parent = SceneComponent->GetAttachParent())
+			{
+				TSharedPtr<FSCSEditorTreeNode> ParentNode = FindOrAddNodeForComponent(Parent);
+
+				if (!ParentNode.IsValid())
+				{
+					return nullptr;
+				}
+
+				TreeView->SetItemExpansion(ParentNode, true);
+
+				TSharedPtr<FSCSEditorTreeNode> ChildNode = ParentNode->AddChildFromComponent(ActorComponent);
+				ObjectToNode.Add(ActorComponent, ChildNode);
+
+				return ChildNode;
+			}
+		}
+		
+		TSharedPtr<FSCSEditorTreeNode> RootNode = FSCSEditorTreeNode::FactoryNodeFromComponent(ActorComponent);
+		RootNodes.Add(RootNode);
+		ObjectToNode.Add(ActorComponent, RootNode);
+
+		TreeView->SetItemExpansion(RootNode, true);
+
+		return RootNode;
+	}
+
+private:
+	bool bIsInEditMode;
+	FOnComponentSelected OnComponentSelected;
+	FIsComponentValid IsComponentValid;
+	TSharedPtr<STreeView<TSharedPtr<FSCSEditorTreeNode>>> TreeView;
+	TMap<FObjectKey, TSharedPtr<FSCSEditorTreeNode>> ObjectToNode;
+	TArray<TSharedPtr<FSCSEditorTreeNode>> RootNodes;
+};
+
+class SActorSequenceEditorWidgetImpl : public SCompoundWidget
+{
+public:
+
+	SLATE_BEGIN_ARGS(SActorSequenceEditorWidgetImpl){}
+	SLATE_END_ARGS();
+
+	~SActorSequenceEditorWidgetImpl()
+	{
+		if (Sequencer.IsValid())
+		{
+			FLevelEditorSequencerIntegration::Get().RemoveSequencer(Sequencer.ToSharedRef());
+			Sequencer->Close();
+			Sequencer = nullptr;
+		}
+
+		GEditor->OnBlueprintPreCompile().Remove(OnBlueprintPreCompileHandle);
+		FCoreUObjectDelegates::OnObjectSaved.Remove(OnObjectSavedHandle);
+	}
+
+	void Construct(const FArguments&, TWeakPtr<FBlueprintEditor> InBlueprintEditor)
+	{
+		OnBlueprintPreCompileHandle = GEditor->OnBlueprintPreCompile().AddSP(this, &SActorSequenceEditorWidgetImpl::OnBlueprintPreCompile);
+		OnObjectSavedHandle = FCoreUObjectDelegates::OnObjectSaved.AddSP(this, &SActorSequenceEditorWidgetImpl::OnObjectPreSave);
+
+		WeakBlueprintEditor = InBlueprintEditor;
+
+		ChildSlot
+		[
+			SAssignNew(Content, SBox)
+			.MinDesiredHeight(200)
+		];
+	}
+
+	FText GetDisplayLabel() const
+	{
+		UActorSequence* Sequence = WeakSequence.Get();
+		return Sequence ? Sequence->GetDisplayName() : LOCTEXT("DefaultSequencerLabel", "Sequencer");
+	}
+
+	UActorSequence* GetActorSequence() const
+	{
+		return WeakSequence.Get();
+	}
+
+	UObject* GetPlaybackContext() const
+	{
+		UActorSequence* LocalActorSequence = GetActorSequence();
+		if (LocalActorSequence)
+		{
+			if (AActor* Actor = LocalActorSequence->GetTypedOuter<AActor>())
+			{
+				return Actor;
+			}
+			else if (UBlueprintGeneratedClass* GeneratedClass = LocalActorSequence->GetTypedOuter<UBlueprintGeneratedClass>())
+			{
+				if (GeneratedClass->SimpleConstructionScript)
+				{
+					return GeneratedClass->SimpleConstructionScript->GetComponentEditorActorInstance();
+				}
+			}
+		}
+		
+		return nullptr;
+	}
+
+	TArray<UObject*> GetEventContexts() const
+	{
+		TArray<UObject*> Contexts;
+		if (auto* Context = GetPlaybackContext())
+		{
+			Contexts.Add(Context);
+		}
+		return Contexts;
+	}
+
+	void SetActorSequence(UActorSequence* NewSequence)
+	{
+		if (UActorSequence* OldSequence = WeakSequence.Get())
+		{
+			if (OnSequenceChangedHandle.IsValid())
+			{
+				OldSequence->OnSignatureChanged().Remove(OnSequenceChangedHandle);
+			}
+		}
+
+		WeakSequence = NewSequence;
+
+		if (NewSequence)
+		{
+			OnSequenceChangedHandle = NewSequence->OnSignatureChanged().AddSP(this, &SActorSequenceEditorWidgetImpl::OnSequenceChanged);
+		}
+
+		// If we already have a sequencer open, just assign the sequence
+		if (Sequencer.IsValid() && NewSequence)
+		{
+			if (Sequencer->GetRootMovieSceneSequence() != NewSequence)
+			{
+				Sequencer->ResetToNewRootSequence(*NewSequence);
+			}
+			return;
+		}
+
+		// If we're setting the sequence to none, destroy sequencer
+		if (!NewSequence)
+		{
+			if (Sequencer.IsValid())
+			{
+				FLevelEditorSequencerIntegration::Get().RemoveSequencer(Sequencer.ToSharedRef());
+				Sequencer->Close();
+				Sequencer = nullptr;
+			}
+
+			Content->SetContent(SNew(STextBlock).Text(LOCTEXT("NothingSelected", "Select a sequence")));
+			return;
+		}
+
+		// We need to initialize a new sequencer instance
+		FSequencerInitParams SequencerInitParams;
+		{
+			TWeakObjectPtr<UActorSequence> LocalWeakSequence = NewSequence;
+
+			SequencerInitParams.RootSequence = NewSequence;
+			SequencerInitParams.EventContexts = TAttribute<TArray<UObject*>>(this, &SActorSequenceEditorWidgetImpl::GetEventContexts);
+			SequencerInitParams.PlaybackContext = TAttribute<UObject*>(this, &SActorSequenceEditorWidgetImpl::GetPlaybackContext);
+
+			TSharedRef<FExtender> AddMenuExtender = MakeShareable(new FExtender);
+
+			AddMenuExtender->AddMenuExtension("AddTracks", EExtensionHook::Before, nullptr,
+				FMenuExtensionDelegate::CreateLambda([=](FMenuBuilder& MenuBuilder){
+
+					MenuBuilder.AddSubMenu(
+						LOCTEXT("AddComponent_Label", "Component"),
+						LOCTEXT("AddComponent_ToolTip", "Add a binding to one of this actor's components and allow it to be animated by Sequencer"),
+						FNewMenuDelegate::CreateRaw(this, &SActorSequenceEditorWidgetImpl::AddPossessComponentMenuExtensions),
+						false /*bInOpenSubMenuOnClick*/,
+						FSlateIcon()//"LevelSequenceEditorStyle", "LevelSequenceEditor.PossessNewActor")
+						);
+
+				})
+			);
+
+			SequencerInitParams.ViewParams.bReadOnly = !WeakBlueprintEditor.IsValid() && !NewSequence->IsEditable();
+			SequencerInitParams.bEditWithinLevelEditor = false;
+			SequencerInitParams.ViewParams.AddMenuExtender = AddMenuExtender;
+			SequencerInitParams.ViewParams.UniqueName = "EmbeddedActorSequenceEditor";
+			SequencerInitParams.ViewParams.ScrubberStyle = ESequencerScrubberStyle::FrameBlock;
+			SequencerInitParams.ViewParams.OnReceivedFocus.BindRaw(this, &SActorSequenceEditorWidgetImpl::OnSequencerReceivedFocus);
+		}
+
+		Sequencer = FModuleManager::LoadModuleChecked<ISequencerModule>("Sequencer").CreateSequencer(SequencerInitParams);
+		Content->SetContent(Sequencer->GetSequencerWidget());
+
+		FLevelEditorSequencerIntegrationOptions Options;
+		Options.bRequiresLevelEvents = true;
+		Options.bRequiresActorEvents = false;
+		Options.bCanRecord = false;
+
+		FLevelEditorSequencerIntegration::Get().AddSequencer(Sequencer.ToSharedRef(), Options);
+	}
+
+	void OnSequencerReceivedFocus()
+	{
+		if (Sequencer.IsValid())
+		{
+			FLevelEditorSequencerIntegration::Get().OnSequencerReceivedFocus(Sequencer.ToSharedRef());
+		}
+	}
+
+	void OnObjectPreSave(UObject* InObject)
+	{
+		TSharedPtr<FBlueprintEditor> BlueprintEditor = WeakBlueprintEditor.Pin();
+		if (Sequencer.IsValid() && BlueprintEditor.IsValid() && InObject && InObject == BlueprintEditor->GetBlueprintObj())
+		{
+			Sequencer->RestorePreAnimatedState();
+		}
+	}
+
+	void OnBlueprintPreCompile(UBlueprint* InBlueprint)
+	{
+		TSharedPtr<FBlueprintEditor> BlueprintEditor = WeakBlueprintEditor.Pin();
+		if (Sequencer.IsValid() && BlueprintEditor.IsValid() && InBlueprint && InBlueprint == BlueprintEditor->GetBlueprintObj())
+		{
+			Sequencer->RestorePreAnimatedState();
+		}
+	}
+
+	void OnSelectionUpdated(TSharedPtr<FSCSEditorTreeNode> SelectedNode)
+	{
+		if (SelectedNode->GetNodeType() != FSCSEditorTreeNode::ComponentNode)
+		{
+			return;
+		}
+
+		UActorComponent* EditingComponent = nullptr;
+
+		TSharedPtr<FBlueprintEditor> BlueprintEditor = WeakBlueprintEditor.Pin();
+		if (BlueprintEditor.IsValid())
+		{
+			UBlueprint* Blueprint = BlueprintEditor->GetBlueprintObj();
+			if (Blueprint)
+			{
+				EditingComponent = SelectedNode->GetEditableComponentTemplate(Blueprint);
+			}
+		}
+		else if (AActor* Actor = GetPreviewActor())
+		{
+			EditingComponent = SelectedNode->FindComponentInstanceInActor(Actor);
+		}
+
+		if (EditingComponent)
+		{
+			const FScopedTransaction Transaction(LOCTEXT("AddComponentToSequencer", "Add component to Sequencer"));
+			Sequencer->GetHandleToObject(EditingComponent, true);
+		}
+
+		FSlateApplication::Get().DismissAllMenus();
+	}
+
+	void AddPossessComponentMenuExtensions(FMenuBuilder& MenuBuilder)
+	{
+		AActor* Actor = GetPreviewActor();
+		if (!Actor)
+		{
+			return;
+		}
+
+		Sequencer->State.ClearObjectCaches(*Sequencer);
+		TSet<UObject*> AllBoundObjects;
+
+		AllBoundObjects.Add(GetOwnerComponent());
+
+		UMovieScene* MovieScene = Sequencer->GetFocusedMovieSceneSequence()->GetMovieScene();
+		for (int32 Index = 0; Index < MovieScene->GetPossessableCount(); ++Index)
+		{
+			FMovieScenePossessable& Possessable = MovieScene->GetPossessable(Index);
+			for (TWeakObjectPtr<> WeakObject : Sequencer->FindBoundObjects(Possessable.GetGuid(), Sequencer->GetFocusedTemplateID()))
+			{
+				if (UObject* Object = WeakObject.Get())
+				{
+					AllBoundObjects.Add(Object);
+				}
+			}
+		}
+
+		bool bIdent = false;
+		MenuBuilder.AddWidget(
+			SNew(SComponentSelectionTree, Actor)
+			.IsInEditMode(WeakBlueprintEditor.Pin().IsValid())
+			.OnComponentSelected(this, &SActorSequenceEditorWidgetImpl::OnSelectionUpdated)
+			.IsComponentValid_Lambda(
+				[AllBoundObjects](UActorComponent* Component)
+				{
+					return !AllBoundObjects.Contains(Component);
+				}
+			)
+			, FText(), !bIdent
+		);
+	}
+
+	AActor* GetPreviewActor() const
+	{
+		TSharedPtr<FBlueprintEditor> BlueprintEditor = WeakBlueprintEditor.Pin();
+		if (BlueprintEditor.IsValid())
+		{
+			return BlueprintEditor->GetPreviewActor();
+		}
+		if (UActorSequence* Sequence = WeakSequence.Get())
+		{
+			return Sequence->GetTypedOuter<AActor>();
+		}
+		return nullptr;
+	}
+
+	UActorComponent* GetOwnerComponent() const
+	{
+		UActorSequence* ActorSequence = WeakSequence.Get();
+		AActor* Actor = ActorSequence ? GetPreviewActor() : nullptr;
+
+		return Actor ? FindObject<UActorComponent>(Actor, *ActorSequence->GetOuter()->GetName()) : nullptr;
+	}
+
+	void OnSequenceChanged()
+	{
+		UActorSequence* ActorSequence = WeakSequence.Get();
+		UBlueprint* Blueprint = ActorSequence ? ActorSequence->GetParentBlueprint() : nullptr;
+
+		if (Blueprint)
+		{
+			FBlueprintEditorUtils::MarkBlueprintAsModified(Blueprint);
+		}
+	}
+
+private:
+	TWeakObjectPtr<UActorSequence> WeakSequence;
+
+	TWeakPtr<FBlueprintEditor> WeakBlueprintEditor;
+
+	TSharedPtr<SBox> Content;
+	TSharedPtr<ISequencer> Sequencer;
+
+	FDelegateHandle OnBlueprintPreCompileHandle;
+	FDelegateHandle OnObjectSavedHandle;
+
+	FDelegateHandle OnSequenceChangedHandle;
+};
+
+void SActorSequenceEditorWidget::Construct(const FArguments&, TWeakPtr<FBlueprintEditor> InBlueprintEditor)
+{
+	ChildSlot
+	[
+		SAssignNew(Impl, SActorSequenceEditorWidgetImpl, InBlueprintEditor)
+	];
+}
+
+FText SActorSequenceEditorWidget::GetDisplayLabel() const
+{
+	return Impl.Pin()->GetDisplayLabel();
+}
+
+void SActorSequenceEditorWidget::AssignSequence(UActorSequence* NewActorSequence)
+{
+	Impl.Pin()->SetActorSequence(NewActorSequence);
+}
+
+UActorSequence* SActorSequenceEditorWidget::GetSequence() const
+{
+	return Impl.Pin()->GetActorSequence();
+}
+
+FActorSequenceEditorSummoner::FActorSequenceEditorSummoner(TSharedPtr<FBlueprintEditor> BlueprintEditor)
+	: FWorkflowTabFactory("EmbeddedSequenceID", BlueprintEditor)
+	, WeakBlueprintEditor(BlueprintEditor)
+{
+	bIsSingleton = true;
+
+	TabLabel = LOCTEXT("SequencerTabName", "Sequencer");
+}
+
+TSharedRef<SWidget> FActorSequenceEditorSummoner::CreateTabBody(const FWorkflowTabSpawnInfo& Info) const
+{
+	return SNew(SActorSequenceEditorWidget, WeakBlueprintEditor);
+}
+
+#undef LOCTEXT_NAMESPACE
diff --git a/Engine/Plugins/MovieScene/ActorSequence/Source/ActorSequenceEditor/Private/ActorSequenceEditorTabSummoner.h b/Engine/Plugins/MovieScene/ActorSequence/Source/ActorSequenceEditor/Private/ActorSequenceEditorTabSummoner.h
new file mode 100644
--- /dev/null
+++ b/Engine/Plugins/MovieScene/ActorSequence/Source/ActorSequenceEditor/Private/ActorSequenceEditorTabSummoner.h
@@ -0,0 +1,44 @@
+// Copyright 1998-2018 Epic Games, Inc. All Rights Reserved.
+
+#pragma once
+
+#include "WorkflowOrientedApp/WorkflowUObjectDocuments.h"
+
+
+class FBlueprintEditor;
+class SActorSequenceEditorWidgetImpl;
+class UActorSequence;
+class UActorSequenceComponent;
+
+
+class SActorSequenceEditorWidget
+	: public SCompoundWidget
+{
+public:
+
+	SLATE_BEGIN_ARGS(SActorSequenceEditorWidget){}
+	SLATE_END_ARGS();
+
+	void Construct(const FArguments&, TWeakPtr<FBlueprintEditor> InBlueprintEditor);
+	void AssignSequence(UActorSequence* NewActorSequence);
+	UActorSequence* GetSequence() const;
+	FText GetDisplayLabel() const;
+
+private:
+
+	TWeakPtr<SActorSequenceEditorWidgetImpl> Impl;
+};
+
+
+struct FActorSequenceEditorSummoner
+	: public FWorkflowTabFactory
+{
+	FActorSequenceEditorSummoner(TSharedPtr<FBlueprintEditor> BlueprintEditor);
+
+	virtual TSharedRef<SWidget> CreateTabBody(const FWorkflowTabSpawnInfo& Info) const override;
+
+protected:
+
+	TWeakObjectPtr<UActorSequenceComponent> WeakComponent;
+	TWeakPtr<FBlueprintEditor> WeakBlueprintEditor;
+};
diff --git a/Engine/Plugins/MovieScene/ActorSequence/Source/ActorSequenceEditor/Private/MovieSceneSequenceEditor_ActorSequence.h b/Engine/Plugins/MovieScene/ActorSequence/Source/ActorSequenceEditor/Private/MovieSceneSequenceEditor_ActorSequence.h
new file mode 100644
--- /dev/null
+++ b/Engine/Plugins/MovieScene/ActorSequence/Source/ActorSequenceEditor/Private/MovieSceneSequenceEditor_ActorSequence.h
@@ -0,0 +1,65 @@
+// Copyright 1998-2018 Epic Games, Inc. All Rights Reserved.
+
+#pragma once
+
+#include "MovieSceneSequenceEditor.h"
+#include "ActorSequence.h"
+#include "Engine/Level.h"
+#include "Engine/LevelScriptBlueprint.h"
+#include "Kismet2/KismetEditorUtilities.h"
+#include "Tracks/MovieSceneEventTrack.h"
+#include "K2Node_FunctionEntry.h"
+#include "EdGraphSchema_K2.h"
+
+struct FMovieSceneSequenceEditor_ActorSequence : FMovieSceneSequenceEditor
+{
+	virtual bool CanCreateEvents(UMovieSceneSequence* InSequence) const
+	{
+		return true;
+	}
+
+	virtual UBlueprint* GetBlueprintForSequence(UMovieSceneSequence* InSequence) const override
+	{
+		UActorSequence* ActorSequence = CastChecked<UActorSequence>(InSequence);
+		if (UBlueprint* Blueprint = ActorSequence->GetParentBlueprint())
+		{
+			return Blueprint;
+		}
+
+		UActorSequenceComponent* Component = ActorSequence->GetTypedOuter<UActorSequenceComponent>();
+		ULevel* Level = Component ? Component->GetOwner()->GetLevel() : nullptr;
+
+		bool bDontCreateNewBlueprint = true;
+		return Level ? Level->GetLevelScriptBlueprint(bDontCreateNewBlueprint) : nullptr;
+	}
+
+	virtual UBlueprint* CreateBlueprintForSequence(UMovieSceneSequence* InSequence) const override
+	{
+		UActorSequence* ActorSequence = CastChecked<UActorSequence>(InSequence);
+		check(!ActorSequence->GetParentBlueprint());
+
+		UActorSequenceComponent* Component = ActorSequence->GetTypedOuter<UActorSequenceComponent>();
+		ULevel* Level = Component ? Component->GetOwner()->GetLevel() : nullptr;
+
+		bool bDontCreateNewBlueprint = false;
+		return Level ? Level->GetLevelScriptBlueprint(bDontCreateNewBlueprint) : nullptr;
+	}
+
+	virtual void SetupDefaultPinForEndpoint(UMovieSceneEventTrack* EventTrack, UK2Node_FunctionEntry* Endpoint) const override
+	{
+		// By default for master tracks we create a pin (that receives the level blueprint for master event tracks, or the event receivers)
+		UClass* PinClass = nullptr;
+
+		// When event receivers are not set, we set the pin type to the type of the object binding
+		if (EventTrack->EventReceivers.Num() == 0)
+		{
+			PinClass = FindTrackObjectBindingClass(EventTrack);
+		}
+
+		FEdGraphPinType PinType;
+		PinType.PinCategory = UEdGraphSchema_K2::PC_Object;
+		PinType.PinSubCategoryObject = PinClass ? PinClass : UObject::StaticClass();
+
+		Endpoint->CreateUserDefinedPin(TargetPinName, PinType, EGPD_Output, true);
+	}
+};
\ No newline at end of file
diff --git a/Engine/Plugins/MovieScene/ActorSequenceEditor/ActorSequenceEditor.uplugin b/Engine/Plugins/MovieScene/ActorSequenceEditor/ActorSequenceEditor.uplugin
deleted file mode 100644
diff --git a/Engine/Plugins/MovieScene/ActorSequenceEditor/Source/ActorSequenceEditor/ActorSequenceEditor.Build.cs b/Engine/Plugins/MovieScene/ActorSequenceEditor/Source/ActorSequenceEditor/ActorSequenceEditor.Build.cs
deleted file mode 100644
diff --git a/Engine/Plugins/MovieScene/ActorSequenceEditor/Source/ActorSequenceEditor/Private/ActorSequenceComponentCustomization.cpp b/Engine/Plugins/MovieScene/ActorSequenceEditor/Source/ActorSequenceEditor/Private/ActorSequenceComponentCustomization.cpp
deleted file mode 100644
diff --git a/Engine/Plugins/MovieScene/ActorSequenceEditor/Source/ActorSequenceEditor/Private/ActorSequenceComponentCustomization.h b/Engine/Plugins/MovieScene/ActorSequenceEditor/Source/ActorSequenceEditor/Private/ActorSequenceComponentCustomization.h
deleted file mode 100644
diff --git a/Engine/Plugins/MovieScene/ActorSequenceEditor/Source/ActorSequenceEditor/Private/ActorSequenceEditor.cpp b/Engine/Plugins/MovieScene/ActorSequenceEditor/Source/ActorSequenceEditor/Private/ActorSequenceEditor.cpp
deleted file mode 100644
diff --git a/Engine/Plugins/MovieScene/ActorSequenceEditor/Source/ActorSequenceEditor/Private/ActorSequenceEditorStyle.h b/Engine/Plugins/MovieScene/ActorSequenceEditor/Source/ActorSequenceEditor/Private/ActorSequenceEditorStyle.h
deleted file mode 100644
diff --git a/Engine/Plugins/MovieScene/ActorSequenceEditor/Source/ActorSequenceEditor/Private/ActorSequenceEditorTabSummoner.cpp b/Engine/Plugins/MovieScene/ActorSequenceEditor/Source/ActorSequenceEditor/Private/ActorSequenceEditorTabSummoner.cpp
deleted file mode 100644
diff --git a/Engine/Plugins/MovieScene/ActorSequenceEditor/Source/ActorSequenceEditor/Private/ActorSequenceEditorTabSummoner.h b/Engine/Plugins/MovieScene/ActorSequenceEditor/Source/ActorSequenceEditor/Private/ActorSequenceEditorTabSummoner.h
deleted file mode 100644
diff --git a/Engine/Plugins/MovieScene/ActorSequenceEditor/Source/ActorSequenceEditor/Private/MovieSceneSequenceEditor_ActorSequence.h b/Engine/Plugins/MovieScene/ActorSequenceEditor/Source/ActorSequenceEditor/Private/MovieSceneSequenceEditor_ActorSequence.h
deleted file mode 100644
diff --git a/Engine/Plugins/MovieScene/MatineeToLevelSequence/Source/MatineeToLevelSequence/Private/MatineeToLevelSequenceModule.cpp b/Engine/Plugins/MovieScene/MatineeToLevelSequence/Source/MatineeToLevelSequence/Private/MatineeToLevelSequenceModule.cpp
--- a/Engine/Plugins/MovieScene/MatineeToLevelSequence/Source/MatineeToLevelSequence/Private/MatineeToLevelSequenceModule.cpp
+++ b/Engine/Plugins/MovieScene/MatineeToLevelSequence/Source/MatineeToLevelSequence/Private/MatineeToLevelSequenceModule.cpp
@@ -30,6 +30,8 @@
 #include "Matinee/InterpTrackColorScale.h"
 #include "Matinee/InterpTrackColorProp.h"
 #include "Matinee/InterpTrackFloatProp.h"
+#include "Matinee/InterpTrackFloatMaterialParam.h"
+#include "Matinee/InterpTrackVectorMaterialParam.h"
 #include "Matinee/InterpTrackMove.h"
 #include "Matinee/InterpTrackAnimControl.h"
 #include "Matinee/InterpTrackSound.h"
@@ -59,6 +61,7 @@
 #include "Tracks/MovieSceneBoolTrack.h"
 #include "Tracks/MovieSceneColorTrack.h"
 #include "Tracks/MovieSceneFloatTrack.h"
+#include "Tracks/MovieSceneMaterialTrack.h"
 #include "Tracks/MovieScene3DTransformTrack.h"
 #include "Tracks/MovieSceneParticleTrack.h"
 #include "Tracks/MovieSceneSkeletalAnimationTrack.h"
@@ -73,6 +76,8 @@
 #include "Framework/Notifications/NotificationManager.h"
 #include "Widgets/Notifications/SNotificationList.h"
 #include "MovieSceneTimeHelpers.h"
+#include "Components/SkeletalMeshComponent.h"
+#include "Engine/StaticMeshActor.h"
 
 #define LOCTEXT_NAMESPACE "MatineeToLevelSequence"
 
@@ -425,6 +430,104 @@ protected:
 		return PropertyTrack;
 	}
 
+	FGuid FindComponentGuid(AActor* GroupActor, UMovieSceneSequence* NewSequence, UMovieScene* NewMovieScene, FGuid PossessableGuid)
+	{
+		FGuid ComponentGUID;
+		// Skeletal and static mesh actors can both have material component tracks, and need to have their mesh component added to sequencer.
+		if (GroupActor->GetClass() == ASkeletalMeshActor::StaticClass())
+		{
+
+			ASkeletalMeshActor* SkelMeshActor = CastChecked<ASkeletalMeshActor>(GroupActor);
+			USkeletalMeshComponent* SkelMeshComponent = SkelMeshActor->GetSkeletalMeshComponent();
+			// In matinee a component may be referenced in multiple material tracks, so check to see if this one is already bound.
+			FGuid FoundGUID = NewSequence->FindPossessableObjectId(*SkelMeshComponent, SkelMeshActor);
+			if (FoundGUID != FGuid())
+			{
+				ComponentGUID = FoundGUID;
+			}
+			else
+			{
+				ComponentGUID = NewMovieScene->AddPossessable(SkelMeshComponent->GetName(), SkelMeshComponent->GetClass());
+				FMovieScenePossessable* ChildPossesable = NewMovieScene->FindPossessable(ComponentGUID);
+				ChildPossesable->SetParent(PossessableGuid);
+				NewSequence->BindPossessableObject(ComponentGUID, *SkelMeshComponent, GroupActor->GetWorld());
+			}
+		}
+		else if (GroupActor->GetClass() == AStaticMeshActor::StaticClass())
+		{
+			AStaticMeshActor* StaticMeshActor = CastChecked<AStaticMeshActor>(GroupActor);
+			UStaticMeshComponent* StaticMeshComponent = StaticMeshActor->GetStaticMeshComponent();
+			FGuid FoundGUID = NewSequence->FindPossessableObjectId(*StaticMeshComponent, StaticMeshActor);
+			if (FoundGUID != FGuid())
+			{
+				ComponentGUID = FoundGUID;
+			}
+			else
+			{
+				ComponentGUID = NewMovieScene->AddPossessable(StaticMeshComponent->GetName(), StaticMeshComponent->GetClass());
+				FMovieScenePossessable* ChildPossesable = NewMovieScene->FindPossessable(ComponentGUID);
+				ChildPossesable->SetParent(PossessableGuid);
+				NewSequence->BindPossessableObject(ComponentGUID, *StaticMeshComponent, GroupActor->GetWorld());
+			}
+		}
+		else
+		{
+			return FGuid();
+		}
+		return ComponentGUID;
+	}
+	
+	template<typename T>
+	void CopyMaterialsToComponents(int32 NumMaterials, FGuid ComponentGuid, UMovieScene* NewMovieScene, T* MatineeMaterialParamTrack)
+	{
+		// One matinee material track can change the same parameter for multiple materials, but sequencer binds them to individual tracks.
+		for (int32 MaterialIndex = 0; MaterialIndex < NumMaterials; ++MaterialIndex)
+		{
+			UMovieSceneComponentMaterialTrack* MaterialTrack = nullptr;
+
+			TArray<UMovieSceneComponentMaterialTrack *> BoundTracks;
+
+			// Find all tracks bound to the component we added.
+			for (const FMovieSceneBinding& Binding : NewMovieScene->GetBindings())
+			{
+				if (Binding.GetObjectGuid() == ComponentGuid)
+				{
+					for (auto Track : Binding.GetTracks())
+					{
+						MaterialTrack = Cast<UMovieSceneComponentMaterialTrack>(Track);
+						if (MaterialTrack)
+						{
+							BoundTracks.Add(MaterialTrack);
+						}
+					}
+					break;
+				}
+			}
+
+			// The material may have already been added to the component, so look first to see if there is a track with the current material index.
+			UMovieSceneComponentMaterialTrack** FoundTrack = BoundTracks.FindByPredicate([MaterialIndex](UMovieSceneComponentMaterialTrack* Track) -> bool { return Track && Track->GetMaterialIndex() == MaterialIndex; });
+			if (FoundTrack)
+			{
+				MaterialTrack = *FoundTrack;
+			}
+
+			if (MaterialTrack)
+			{
+				FMatineeImportTools::CopyInterpMaterialParamTrack(MatineeMaterialParamTrack, MaterialTrack);
+			}
+			else
+			{
+				MaterialTrack = NewMovieScene->AddTrack<UMovieSceneComponentMaterialTrack>(ComponentGuid);
+				if (MaterialTrack)
+				{
+					MaterialTrack->SetMaterialIndex(MaterialIndex);
+					FMatineeImportTools::CopyInterpMaterialParamTrack(MatineeMaterialParamTrack, MaterialTrack);
+				}
+			}
+			
+		}
+	}
+
 	/** Convert an interp group */
 	void ConvertInterpGroup(UInterpGroup* Group, AActor* GroupActor, IMovieScenePlayer& Player, UMovieSceneSequence* NewSequence, UMovieScene* NewMovieScene, int32& NumWarnings)
 	{
@@ -553,6 +656,38 @@ protected:
 					}
 				}
 			}
+			else if (Track->IsA(UInterpTrackFloatMaterialParam::StaticClass()))
+			{
+				UInterpTrackFloatMaterialParam* MatineeMaterialParamTrack = StaticCast<UInterpTrackFloatMaterialParam*>(Track);
+				if (MatineeMaterialParamTrack->GetNumKeyframes() != 0 && GroupActor && PossessableGuid.IsValid())
+				{
+					FGuid ComponentGuid = FindComponentGuid(GroupActor, NewSequence, NewMovieScene, PossessableGuid);
+
+					if (ComponentGuid == FGuid())
+					{
+						continue;
+					}
+
+					CopyMaterialsToComponents(MatineeMaterialParamTrack->TargetMaterials.Num(), ComponentGuid, NewMovieScene, MatineeMaterialParamTrack);
+					
+				}
+			}
+			else if (Track->IsA(UInterpTrackVectorMaterialParam::StaticClass()))
+			{
+				UInterpTrackVectorMaterialParam* MatineeMaterialParamTrack = StaticCast<UInterpTrackVectorMaterialParam*>(Track);
+				if (MatineeMaterialParamTrack->GetNumKeyframes() != 0 && GroupActor && PossessableGuid.IsValid())
+				{
+					FGuid ComponentGuid = FindComponentGuid(GroupActor, NewSequence, NewMovieScene, PossessableGuid);
+
+					if (ComponentGuid == FGuid())
+					{
+						continue;
+					}
+
+					CopyMaterialsToComponents(MatineeMaterialParamTrack->TargetMaterials.Num(), ComponentGuid, NewMovieScene, MatineeMaterialParamTrack);
+
+				}
+			}
 			else if (Track->IsA(UInterpTrackVectorProp::StaticClass()))
 			{
 				UInterpTrackVectorProp* MatineeVectorTrack = StaticCast<UInterpTrackVectorProp*>(Track);
diff --git a/Engine/Plugins/MovieScene/SequencerScripting/Source/SequencerScripting/Private/ExtensionLibraries/MovieSceneSequenceExtensions.cpp b/Engine/Plugins/MovieScene/SequencerScripting/Source/SequencerScripting/Private/ExtensionLibraries/MovieSceneSequenceExtensions.cpp
--- a/Engine/Plugins/MovieScene/SequencerScripting/Source/SequencerScripting/Private/ExtensionLibraries/MovieSceneSequenceExtensions.cpp
+++ b/Engine/Plugins/MovieScene/SequencerScripting/Source/SequencerScripting/Private/ExtensionLibraries/MovieSceneSequenceExtensions.cpp
@@ -24,14 +24,14 @@ TArray<UMovieSceneTrack*> UMovieSceneSequenceExtensions::FilterTracks(TArrayView
 
 UMovieScene* UMovieSceneSequenceExtensions::GetMovieScene(UMovieSceneSequence* Sequence)
 {
-	return Sequence->GetMovieScene();
+	return Sequence ? Sequence->GetMovieScene() : nullptr;
 }
 
 TArray<UMovieSceneTrack*> UMovieSceneSequenceExtensions::GetMasterTracks(UMovieSceneSequence* Sequence)
 {
 	TArray<UMovieSceneTrack*> Tracks;
 
-	UMovieScene* MovieScene = Sequence->GetMovieScene();
+	UMovieScene* MovieScene = GetMovieScene(Sequence);
 	if (MovieScene)
 	{
 		Tracks = MovieScene->GetMasterTracks();
@@ -47,7 +47,7 @@ TArray<UMovieSceneTrack*> UMovieSceneSequenceExtensions::GetMasterTracks(UMovieS
 
 TArray<UMovieSceneTrack*> UMovieSceneSequenceExtensions::FindMasterTracksByType(UMovieSceneSequence* Sequence, TSubclassOf<UMovieSceneTrack> TrackType)
 {
-	UMovieScene* MovieScene   = Sequence->GetMovieScene();
+	UMovieScene* MovieScene   = GetMovieScene(Sequence);
 	UClass*      DesiredClass = TrackType.Get();
 
 	if (MovieScene && DesiredClass)
@@ -70,7 +70,7 @@ TArray<UMovieSceneTrack*> UMovieSceneSequenceExtensions::FindMasterTracksByType(
 
 TArray<UMovieSceneTrack*> UMovieSceneSequenceExtensions::FindMasterTracksByExactType(UMovieSceneSequence* Sequence, TSubclassOf<UMovieSceneTrack> TrackType)
 {
-	UMovieScene* MovieScene   = Sequence->GetMovieScene();
+	UMovieScene* MovieScene   = GetMovieScene(Sequence);
 	UClass*      DesiredClass = TrackType.Get();
 
 	if (MovieScene && DesiredClass)
@@ -94,7 +94,7 @@ TArray<UMovieSceneTrack*> UMovieSceneSequenceExtensions::FindMasterTracksByExact
 UMovieSceneTrack* UMovieSceneSequenceExtensions::AddMasterTrack(UMovieSceneSequence* Sequence, TSubclassOf<UMovieSceneTrack> TrackType)
 {
 	// @todo: sequencer-python: master track type compatibility with sequence. Currently that's really only loosely defined by track editors, which is not sufficient here.
-	UMovieScene* MovieScene = Sequence->GetMovieScene();
+	UMovieScene* MovieScene = GetMovieScene(Sequence);
 	if (MovieScene)
 	{
 		UMovieSceneTrack* NewTrack = MovieScene->AddMasterTrack(TrackType);
@@ -109,16 +109,34 @@ UMovieSceneTrack* UMovieSceneSequenceExtensions::AddMasterTrack(UMovieSceneSeque
 
 FFrameRate UMovieSceneSequenceExtensions::GetDisplayRate(UMovieSceneSequence* Sequence)
 {
-	UMovieScene* MovieScene = Sequence->GetMovieScene();
+	UMovieScene* MovieScene = GetMovieScene(Sequence);
 	return MovieScene ? MovieScene->GetDisplayRate() : FFrameRate();
 }
 
+void UMovieSceneSequenceExtensions::SetDisplayRate(UMovieSceneSequence* Sequence, FFrameRate DisplayRate)
+{
+	UMovieScene* MovieScene = GetMovieScene(Sequence);
+	if (MovieScene)
+	{
+		MovieScene->SetDisplayRate(DisplayRate);
+	}
+}
+
 FFrameRate UMovieSceneSequenceExtensions::GetTickResolution(UMovieSceneSequence* Sequence)
 {
-	UMovieScene* MovieScene = Sequence->GetMovieScene();
+	UMovieScene* MovieScene = GetMovieScene(Sequence);
 	return MovieScene ? MovieScene->GetTickResolution() : FFrameRate();
 }
 
+void UMovieSceneSequenceExtensions::SetTickResolution(UMovieSceneSequence* Sequence, FFrameRate TickResolution)
+{
+	UMovieScene* MovieScene = GetMovieScene(Sequence);
+	if (MovieScene)
+	{
+		MovieScene->SetTickResolutionDirectly(TickResolution);
+	}
+}
+
 FSequencerScriptingRange UMovieSceneSequenceExtensions::MakeRange(UMovieSceneSequence* Sequence, int32 StartFrame, int32 Duration)
 {
 	FFrameRate FrameRate = GetDisplayRate(Sequence);
@@ -133,7 +151,7 @@ FSequencerScriptingRange UMovieSceneSequenceExtensions::MakeRangeSeconds(UMovieS
 
 FSequencerScriptingRange UMovieSceneSequenceExtensions::GetPlaybackRange(UMovieSceneSequence* Sequence)
 {
-	UMovieScene* MovieScene = Sequence->GetMovieScene();
+	UMovieScene* MovieScene = GetMovieScene(Sequence);
 	if (MovieScene)
 	{
 		return FSequencerScriptingRange::FromNative(MovieScene->GetPlaybackRange(), GetTickResolution(Sequence));
@@ -146,7 +164,7 @@ FSequencerScriptingRange UMovieSceneSequenceExtensions::GetPlaybackRange(UMovieS
 
 FTimecode UMovieSceneSequenceExtensions::GetTimecodeSource(UMovieSceneSequence* Sequence)
 {
-	UMovieScene* MovieScene = Sequence->GetMovieScene();
+	UMovieScene* MovieScene = GetMovieScene(Sequence);
 	if (MovieScene)
 	{
 		return MovieScene->TimecodeSource.Timecode;
@@ -159,7 +177,7 @@ FTimecode UMovieSceneSequenceExtensions::GetTimecodeSource(UMovieSceneSequence*
 
 FSequencerBindingProxy UMovieSceneSequenceExtensions::FindBindingByName(UMovieSceneSequence* Sequence, FString Name)
 {
-	UMovieScene* MovieScene = Sequence->GetMovieScene();
+	UMovieScene* MovieScene = GetMovieScene(Sequence);
 	if (MovieScene)
 	{
 		const FMovieSceneBinding* Binding = Algo::FindBy(MovieScene->GetBindings(), Name, &FMovieSceneBinding::GetName);
@@ -175,7 +193,7 @@ TArray<FSequencerBindingProxy> UMovieSceneSequenceExtensions::GetBindings(UMovie
 {
 	TArray<FSequencerBindingProxy> AllBindings;
 
-	UMovieScene* MovieScene = Sequence->GetMovieScene();
+	UMovieScene* MovieScene = GetMovieScene(Sequence);
 	if (MovieScene)
 	{
 		for (const FMovieSceneBinding& Binding : MovieScene->GetBindings())
@@ -191,7 +209,7 @@ TArray<FSequencerBindingProxy> UMovieSceneSequenceExtensions::GetSpawnables(UMov
 {
 	TArray<FSequencerBindingProxy> AllSpawnables;
 
-	UMovieScene* MovieScene = Sequence->GetMovieScene();
+	UMovieScene* MovieScene = GetMovieScene(Sequence);
 	if (MovieScene)
 	{
 		int32 Count = MovieScene->GetSpawnableCount();
@@ -209,7 +227,7 @@ TArray<FSequencerBindingProxy> UMovieSceneSequenceExtensions::GetPossessables(UM
 {
 	TArray<FSequencerBindingProxy> AllPossessables;
 
-	UMovieScene* MovieScene = Sequence->GetMovieScene();
+	UMovieScene* MovieScene = GetMovieScene(Sequence);
 	if (MovieScene)
 	{
 		int32 Count = MovieScene->GetPossessableCount();
diff --git a/Engine/Plugins/MovieScene/SequencerScripting/Source/SequencerScripting/Public/ExtensionLibraries/MovieSceneSequenceExtensions.h b/Engine/Plugins/MovieScene/SequencerScripting/Source/SequencerScripting/Public/ExtensionLibraries/MovieSceneSequenceExtensions.h
--- a/Engine/Plugins/MovieScene/SequencerScripting/Source/SequencerScripting/Public/ExtensionLibraries/MovieSceneSequenceExtensions.h
+++ b/Engine/Plugins/MovieScene/SequencerScripting/Source/SequencerScripting/Public/ExtensionLibraries/MovieSceneSequenceExtensions.h
@@ -70,7 +70,7 @@ public:
 	static UMovieSceneTrack* AddMasterTrack(UMovieSceneSequence* Sequence, TSubclassOf<UMovieSceneTrack> TrackType);
 
 	/**
-	 * Get's this sequence's display rate
+	 * Gets this sequence's display rate
 	 *
 	 * @param Sequence        The sequence to use
 	 * @return The display rate that this sequence is displayed as
@@ -79,7 +79,16 @@ public:
 	static FFrameRate GetDisplayRate(UMovieSceneSequence* Sequence);
 
 	/**
-	 * Get's this sequence's tick resolution
+	 * Sets this sequence's display rate
+	 *
+	 * @param Sequence        The sequence to use
+	 * @param DisplayRate The display rate that this sequence is displayed as
+	 */
+	UFUNCTION(BlueprintCallable, Category="Sequence", meta=(ScriptMethod))
+	static void SetDisplayRate(UMovieSceneSequence* Sequence, FFrameRate DisplayRate);
+
+	/**
+	 * Gets this sequence's tick resolution
 	 *
 	 * @param Sequence        The sequence to use
 	 * @return The tick resolution of the sequence, defining the smallest unit of time representable on this sequence
@@ -87,6 +96,15 @@ public:
 	UFUNCTION(BlueprintCallable, Category="Sequence", meta=(ScriptMethod))
 	static FFrameRate GetTickResolution(UMovieSceneSequence* Sequence);
 
+	/**
+	 * Sets this sequence's tick resolution
+	 *
+	 * @param Sequence        The sequence to use
+	 * @param TickResolution The tick resolution of the sequence, defining the smallest unit of time representable on this sequence
+	 */
+	UFUNCTION(BlueprintCallable, Category="Sequence", meta=(ScriptMethod))
+	static void SetTickResolution(UMovieSceneSequence* Sequence, FFrameRate TickResolution);
+
 	/**
 	 * Make a new range for this sequence in its display rate
 	 *
@@ -119,11 +137,11 @@ public:
 	static FSequencerScriptingRange GetPlaybackRange(UMovieSceneSequence* Sequence);
 
 	/**
-	* Get the timecode source of this sequence
-	*
-	* @param Sequence        The sequence within which to get the timecode source
-	* @return Timecode source of this sequence
-	*/
+	 * Get the timecode source of this sequence
+	 *
+	 * @param Sequence        The sequence within which to get the timecode source
+	 * @return Timecode source of this sequence
+	 */
 	UFUNCTION(BlueprintCallable, Category = "Sequence", meta = (ScriptMethod))
 	static FTimecode GetTimecodeSource(UMovieSceneSequence* Sequence);
 
diff --git a/Engine/Plugins/Runtime/AudioCapture/Source/AudioCaptureEditor/Private/AudioRecordingManager.cpp b/Engine/Plugins/Runtime/AudioCapture/Source/AudioCaptureEditor/Private/AudioRecordingManager.cpp
--- a/Engine/Plugins/Runtime/AudioCapture/Source/AudioCaptureEditor/Private/AudioRecordingManager.cpp
+++ b/Engine/Plugins/Runtime/AudioCapture/Source/AudioCaptureEditor/Private/AudioRecordingManager.cpp
@@ -252,7 +252,7 @@ void FAudioRecordingManager::StopRecording(TArray<USoundWave*>& OutSoundWaves)
 
 			UE_LOG(LogMicManager, Log, TEXT("Stopping mic recording. Recorded %d frames of audio (%.4f seconds). Detected %d buffer overflows."), 
 				NumRecordedSamples, 
-				(float)NumRecordedSamples / RecordingSampleRate,
+				((float)NumRecordedSamples / RecordingSampleRate) / (float)NumInputChannels,
 				NumOverflowsDetected);
 
 			// Get a ptr to the buffer we're actually going to serialize
diff --git a/Engine/Source/Editor/MovieSceneTools/Private/AutomatedLevelSequenceCapture.cpp b/Engine/Source/Editor/MovieSceneTools/Private/AutomatedLevelSequenceCapture.cpp
--- a/Engine/Source/Editor/MovieSceneTools/Private/AutomatedLevelSequenceCapture.cpp
+++ b/Engine/Source/Editor/MovieSceneTools/Private/AutomatedLevelSequenceCapture.cpp
@@ -208,6 +208,7 @@ void UAutomatedLevelSequenceCapture::Initialize(TSharedPtr<FSceneViewport> InVie
 		Actor->PlaybackSettings.LoopCount.Value = 0;
 		Actor->SequencePlayer->SetTimeController(MakeShared<FMovieSceneTimeController_FrameStep>());
 		Actor->PlaybackSettings.bPauseAtEnd = true;
+		Actor->PlaybackSettings.bAutoPlay = false;
 
 		if (BurnInOptions)
 		{
@@ -231,7 +232,6 @@ void UAutomatedLevelSequenceCapture::Initialize(TSharedPtr<FSceneViewport> InVie
 			}
 			Actor->InitializePlayer();
 		}
-		Actor->bAutoPlay = false;
 
 		if (InitializeShots())
 		{
diff --git a/Engine/Source/Editor/MovieSceneTools/Private/MatineeImportTools.cpp b/Engine/Source/Editor/MovieSceneTools/Private/MatineeImportTools.cpp
--- a/Engine/Source/Editor/MovieSceneTools/Private/MatineeImportTools.cpp
+++ b/Engine/Source/Editor/MovieSceneTools/Private/MatineeImportTools.cpp
@@ -12,10 +12,14 @@
 #include "Channels/MovieSceneEvent.h"
 #include "K2Node_FunctionEntry.h"
 #include "Kismet2/KismetEditorUtilities.h"
+#include "Animation/SkeletalMeshActor.h"
+#include "EngineUtils.h"
 
 #include "Matinee/MatineeActor.h"
 #include "Matinee/InterpData.h"
 #include "Matinee/InterpGroupInst.h"
+#include "Matinee/InterpTrackFloatMaterialParam.h"
+#include "Matinee/InterpTrackVectorMaterialParam.h"
 #include "Matinee/InterpTrackLinearColorProp.h"
 #include "Matinee/InterpTrackColorProp.h"
 #include "Matinee/InterpTrackBoolProp.h"
@@ -30,6 +34,7 @@
 
 #include "Tracks/MovieSceneBoolTrack.h"
 #include "Tracks/MovieSceneFloatTrack.h"
+#include "Tracks/MovieSceneMaterialTrack.h"
 #include "Tracks/MovieSceneColorTrack.h"
 #include "Tracks/MovieScene3DTransformTrack.h"
 #include "Tracks/MovieSceneParticleTrack.h"
@@ -52,6 +57,7 @@
 #include "Sections/MovieSceneEventSection.h"
 #include "Sections/MovieSceneEventTriggerSection.h"
 #include "Sections/MovieSceneVectorSection.h"
+#include "Sections/MovieSceneParameterSection.h"
 
 
 #include "Animation/AnimSequence.h"
@@ -120,7 +126,7 @@ bool FMatineeImportTools::TryConvertMatineeToggleToOutParticleKey( ETrackToggleA
 
 
 void FMatineeImportTools::SetOrAddKey(TMovieSceneChannelData<FMovieSceneFloatValue>& ChannelData, FFrameNumber Time, float Value, float ArriveTangent, float LeaveTangent, EInterpCurveMode MatineeInterpMode, FFrameRate FrameRate
-	, ERichCurveTangentWeightMode WeightedMode, float ArriveTangentWeight, float LeaveTangentWeight )
+	, ERichCurveTangentWeightMode WeightedMode, float ArriveTangentWeight, float LeaveTangentWeight)
 {
 	if (ChannelData.FindKey(Time) == INDEX_NONE)
 	{
@@ -224,6 +230,131 @@ bool FMatineeImportTools::CopyInterpFloatTrack( UInterpTrackFloatBase* MatineeFl
 	return bSectionCreated;
 }
 
+bool FMatineeImportTools::CopyInterpMaterialParamTrack(UInterpTrackFloatMaterialParam* MatineeMaterialParamTrack, UMovieSceneComponentMaterialTrack * MaterialTrack)
+{
+	const FScopedTransaction Transaction(NSLOCTEXT("Sequencer", "PasteMatineeFloatMaterialParamTrack", "Paste Matinee Float Material Param Track"));
+	bool bSectionCreated = false;
+
+	MaterialTrack->Modify();
+
+	FFrameRate   FrameRate = MaterialTrack->GetTypedOuter<UMovieScene>()->GetTickResolution();
+	FFrameNumber FirstKeyTime = (MatineeMaterialParamTrack->GetKeyframeTime(0) * FrameRate).RoundToFrame();
+
+	UMovieSceneParameterSection* Section = Cast<UMovieSceneParameterSection>(MovieSceneHelpers::FindSectionAtTime(MaterialTrack->GetAllSections(), FirstKeyTime));
+	if (Section == nullptr)
+	{
+		Section = Cast<UMovieSceneParameterSection>(MaterialTrack->CreateNewSection());
+		MaterialTrack->AddSection(*Section);
+		Section->SetRange(TRange<FFrameNumber>::All());
+		bSectionCreated = true;
+	}
+	if (Section->TryModify())
+	{
+		auto FirstPoint = MatineeMaterialParamTrack->FloatTrack.Points[0];
+		FFrameNumber KeyTime = (FirstPoint.InVal * FrameRate).RoundToFrame();
+
+		// The section needs a key added to be initialized, so add the first key from the matinee track.
+		Section->AddScalarParameterKey(MatineeMaterialParamTrack->ParamName, KeyTime, FirstPoint.OutVal);
+
+		TArray<FScalarParameterNameAndCurve> ScalarParms = Section->GetScalarParameterNamesAndCurves();
+
+		FMovieSceneFloatChannel* Channel = nullptr;
+		int32 CurveIndex = 0;
+		for (FScalarParameterNameAndCurve NameAndCurve : ScalarParms)
+		{
+			if (NameAndCurve.ParameterName == MatineeMaterialParamTrack->ParamName)
+			{
+				Channel = Section->GetChannelProxy().GetChannel<FMovieSceneFloatChannel>(CurveIndex);
+				break;
+			}
+			++CurveIndex;
+		}
+
+		if (!Channel)
+		{
+			return false;
+		}
+
+		
+		TMovieSceneChannelData<FMovieSceneFloatValue> ChannelData = Channel->GetData();
+
+
+		for (const auto& Point : MatineeMaterialParamTrack->FloatTrack.Points)
+		{
+			FFrameNumber CurrentKeyTime = (Point.InVal * FrameRate).RoundToFrame();
+
+			FMatineeImportTools::SetOrAddKey(ChannelData, CurrentKeyTime, Point.OutVal, Point.ArriveTangent, Point.LeaveTangent, Point.InterpMode, FrameRate);
+		}
+
+		CleanupCurveKeys(Channel);
+	}
+
+	return bSectionCreated;
+}
+
+bool FMatineeImportTools::CopyInterpMaterialParamTrack(UInterpTrackVectorMaterialParam * MatineeMaterialParamTrack, UMovieSceneComponentMaterialTrack * MaterialTrack)
+{
+	const FScopedTransaction Transaction(NSLOCTEXT("Sequencer", "PasteMatineeFloatMaterialParamTrack", "Paste Matinee Float Material Param Track"));
+	bool bSectionCreated = false;
+
+	MaterialTrack->Modify();
+
+	FFrameRate   FrameRate = MaterialTrack->GetTypedOuter<UMovieScene>()->GetTickResolution();
+	FFrameNumber FirstKeyTime = (MatineeMaterialParamTrack->GetKeyframeTime(0) * FrameRate).RoundToFrame();
+
+	UMovieSceneParameterSection* Section = Cast<UMovieSceneParameterSection>(MovieSceneHelpers::FindSectionAtTime(MaterialTrack->GetAllSections(), FirstKeyTime));
+	if (Section == nullptr)
+	{
+		Section = Cast<UMovieSceneParameterSection>(MaterialTrack->CreateNewSection());
+		MaterialTrack->AddSection(*Section);
+		Section->SetRange(TRange<FFrameNumber>::All());
+		bSectionCreated = true;
+	}
+	if (Section->TryModify())
+	{
+		FInterpCurvePointVector FirstPoint = MatineeMaterialParamTrack->VectorTrack.Points[0];
+		FFrameNumber KeyTime = (FirstPoint.InVal * FrameRate).RoundToFrame();
+
+		// The section needs a key added to be initialized, so add the first key from the matinee track.
+		Section->AddVectorParameterKey(MatineeMaterialParamTrack->ParamName, KeyTime, FirstPoint.OutVal);
+
+		TArray<FScalarParameterNameAndCurve> ScalarParms = Section->GetScalarParameterNamesAndCurves();
+
+		FMovieSceneFloatChannel* Channel = nullptr;
+		int32 CurveIndex = 0;
+		for (FScalarParameterNameAndCurve NameAndCurve : ScalarParms)
+		{
+			if (NameAndCurve.ParameterName == MatineeMaterialParamTrack->ParamName)
+			{
+				Channel = Section->GetChannelProxy().GetChannel<FMovieSceneFloatChannel>(CurveIndex);
+				break;
+			}
+			++CurveIndex;
+		}
+
+		TArrayView<FMovieSceneFloatChannel*> Channels = Section->GetChannelProxy().GetChannels<FMovieSceneFloatChannel>();
+		TMovieSceneChannelData<FMovieSceneFloatValue> ChannelData[3] = {
+			Channels[CurveIndex]->GetData(), Channels[CurveIndex + 1]->GetData(), Channels[CurveIndex + 2]->GetData()
+		};
+
+
+		for (const auto& Point : MatineeMaterialParamTrack->VectorTrack.Points)
+		{
+			FFrameNumber CurrentKeyTime = (Point.InVal * FrameRate).RoundToFrame();
+
+			FMatineeImportTools::SetOrAddKey(ChannelData[0], CurrentKeyTime, Point.OutVal.X, Point.ArriveTangent.X, Point.LeaveTangent.X, Point.InterpMode, FrameRate);
+			FMatineeImportTools::SetOrAddKey(ChannelData[1], CurrentKeyTime, Point.OutVal.Y, Point.ArriveTangent.Y, Point.LeaveTangent.Y, Point.InterpMode, FrameRate);
+			FMatineeImportTools::SetOrAddKey(ChannelData[2], CurrentKeyTime, Point.OutVal.Z, Point.ArriveTangent.Z, Point.LeaveTangent.Z, Point.InterpMode, FrameRate);
+		}
+
+		CleanupCurveKeys(Channels[0]);
+		CleanupCurveKeys(Channels[1]);
+		CleanupCurveKeys(Channels[2]);
+	}
+
+	return bSectionCreated;
+}
+
 bool FMatineeImportTools::CopyInterpVectorTrack( UInterpTrackVectorProp* MatineeVectorTrack, UMovieSceneVectorTrack* VectorTrack )
 {
 	const FScopedTransaction Transaction( NSLOCTEXT( "Sequencer", "PasteMatineeVectorTrack", "Paste Matinee Vector Track" ) );
@@ -499,7 +630,7 @@ bool FMatineeImportTools::CopyInterpMoveTrack( UInterpTrackMove* MoveTrack, UMov
 				}
 			}
 		}
-
+		
 		CleanupCurveKeys(Channels[0]);
 		CleanupCurveKeys(Channels[1]);
 		CleanupCurveKeys(Channels[2]);
diff --git a/Engine/Source/Editor/MovieSceneTools/Private/MovieSceneToolHelpers.cpp b/Engine/Source/Editor/MovieSceneTools/Private/MovieSceneToolHelpers.cpp
--- a/Engine/Source/Editor/MovieSceneTools/Private/MovieSceneToolHelpers.cpp
+++ b/Engine/Source/Editor/MovieSceneTools/Private/MovieSceneToolHelpers.cpp
@@ -890,6 +890,7 @@ bool ImportFBXProperty(FString NodeName, FString AnimatedPropertyName, FGuid Obj
 
 			if (FloatTrack)
 			{
+				FloatTrack->Modify();
 				FloatTrack->RemoveAllAnimationData();
 
 				FFrameRate FrameRate = FloatTrack->GetTypedOuter<UMovieScene>()->GetTickResolution();
@@ -965,10 +966,14 @@ bool ImportFBXProperty(FString NodeName, FString AnimatedPropertyName, FGuid Obj
 
 void ImportTransformChannel(const FRichCurve& Source, FMovieSceneFloatChannel* Dest, FFrameRate DestFrameRate, bool bNegateTangents)
 {
+	// If there are no keys, don't clear the existing channel
+	if (!Source.GetNumKeys())
+	{
+		return;
+	}
+
 	TMovieSceneChannelData<FMovieSceneFloatValue> ChannelData = Dest->GetData();
 	ChannelData.Reset();
-	double DecimalRate = DestFrameRate.AsDecimal();
-
 	for (auto SourceIt = Source.GetKeyHandleIterator(); SourceIt; ++SourceIt)
 	{
 		const FRichCurveKey Key = Source.GetKey(SourceIt.Key());
@@ -977,7 +982,7 @@ void ImportTransformChannel(const FRichCurve& Source, FMovieSceneFloatChannel* D
 		if (Source.IsKeyHandleValid(PrevKeyHandle))
 		{
 			const FRichCurveKey PrevKey = Source.GetKey(PrevKeyHandle);
-			ArriveTangent = ArriveTangent / ((Key.Time - PrevKey.Time) * DecimalRate);
+			ArriveTangent = ArriveTangent / (Key.Time - PrevKey.Time);
 
 		}
 		float LeaveTangent = Key.LeaveTangent;
@@ -985,7 +990,7 @@ void ImportTransformChannel(const FRichCurve& Source, FMovieSceneFloatChannel* D
 		if (Source.IsKeyHandleValid(NextKeyHandle))
 		{
 			const FRichCurveKey NextKey = Source.GetKey(NextKeyHandle);
-			LeaveTangent = LeaveTangent / ((NextKey.Time - Key.Time) * DecimalRate);
+			LeaveTangent = LeaveTangent / (NextKey.Time - Key.Time);
 		}
 
 		if (bNegateTangents)
@@ -1030,7 +1035,7 @@ bool ImportFBXTransform(FString NodeName, FGuid ObjectBinding, UnFbx::FFbxCurves
 		InMovieScene->Modify();
 		TransformTrack = InMovieScene->AddTrack<UMovieScene3DTransformTrack>(ObjectBinding);
 	}
-	TransformTrack->RemoveAllAnimationData();
+	TransformTrack->Modify();
 
 	bool bSectionAdded = false;
 	UMovieScene3DTransformSection* TransformSection = Cast<UMovieScene3DTransformSection>(TransformTrack->FindOrAddSection(0, bSectionAdded));
@@ -1043,7 +1048,6 @@ bool ImportFBXTransform(FString NodeName, FGuid ObjectBinding, UnFbx::FFbxCurves
 
 	FFrameRate FrameRate = TransformSection->GetTypedOuter<UMovieScene>()->GetTickResolution();
 
-
 	if (bSectionAdded)
 	{
 		TransformSection->SetRange(TRange<FFrameNumber>::All());
@@ -1080,25 +1084,8 @@ bool ImportFBXTransform(FString NodeName, FGuid ObjectBinding, UnFbx::FFbxCurves
 	return true;
 }
 
-bool ImportFBXNode(FString NodeName, UnFbx::FFbxCurvesAPI& CurveAPI, UMovieScene* InMovieScene, ISequencer& InSequencer, const TMap<FGuid, FString>& InObjectBindingMap, bool bMatchByNameOnly)
+bool ImportFBXNode(FString NodeName, UnFbx::FFbxCurvesAPI& CurveAPI, UMovieScene* InMovieScene, ISequencer& InSequencer, FGuid ObjectBinding)
 {
-	// Find the matching object binding to apply this animation to. If not matching by name only, default to the first.
-	FGuid ObjectBinding;
-	for (auto It = InObjectBindingMap.CreateConstIterator(); It; ++It)
-	{
-		if (!bMatchByNameOnly || FCString::Strcmp(*It.Value().ToUpper(), *NodeName.ToUpper()) == 0)
-		{
-			ObjectBinding = It.Key();
-			break;
-		}
-	}
-
-	if (!ObjectBinding.IsValid())
-	{
-		UE_LOG(LogMovieScene, Warning, TEXT("Fbx Import: Failed to find any matching node for (%s)."), *NodeName);
-		return false;
-	}
-
 	// Look for animated float properties
 	TArray<FString> AnimatedPropertyNames;
 	CurveAPI.GetNodeAnimatedPropertyNameArray(NodeName, AnimatedPropertyNames);
@@ -1203,6 +1190,7 @@ void CopyCameraProperties(FbxCamera* CameraNode, ACineCameraActor* CameraActor)
 	CineCameraComponent->SetFieldOfView(FieldOfView);
 	CineCameraComponent->FilmbackSettings.SensorWidth = FUnitConversion::Convert(ApertureWidth, EUnit::Inches, EUnit::Millimeters);
 	CineCameraComponent->FilmbackSettings.SensorHeight = FUnitConversion::Convert(ApertureHeight, EUnit::Inches, EUnit::Millimeters);
+	CineCameraComponent->FocusSettings.ManualFocusDistance = CameraNode->FocusDistance;
 	if (FocalLength < CineCameraComponent->LensSettings.MinFocalLength)
 	{
 		CineCameraComponent->LensSettings.MinFocalLength = FocalLength;
@@ -1365,6 +1353,7 @@ void ImportFBXCamera(UnFbx::FFbxImporter* FbxImporter, UMovieScene* InMovieScene
 				UMovieSceneFloatTrack* FloatTrack = InMovieScene->FindTrack<UMovieSceneFloatTrack>(PropertyOwnerGuid, TEXT("CurrentFocalLength"));
 				if (FloatTrack)
 				{
+					FloatTrack->Modify();
 					FloatTrack->RemoveAllAnimationData();
 
 					bool bSectionAdded = false;
@@ -1574,9 +1563,57 @@ private:
 		TArray<FString> AllNodeNames;
 		CurveAPI.GetAllNodeNameArray(AllNodeNames);
 
+		// First try matching by name
+		for (int32 NodeIndex = 0; NodeIndex < AllNodeNames.Num(); )
+		{
+			FString NodeName = AllNodeNames[NodeIndex];
+			bool bFoundMatch = false;
+			for (auto It = ObjectBindingMap.CreateConstIterator(); It; ++It)
+			{
+				if (FCString::Strcmp(*It.Value().ToUpper(), *NodeName.ToUpper()) == 0)
+				{
+					ImportFBXNode(NodeName, CurveAPI, MovieScene, *Sequencer, It.Key());
+
+					ObjectBindingMap.Remove(It.Key());
+					AllNodeNames.RemoveAt(NodeIndex);
+
+					bFoundMatch = true;
+					break;
+				}
+			}
+
+			if (bFoundMatch)
+			{
+				continue;
+			}
+
+			++NodeIndex;
+		}
+
+		// Otherwise, get the first available node that hasn't been imported onto yet
+		if (!bMatchByNameOnly)
+		{
+			for (int32 NodeIndex = 0; NodeIndex < AllNodeNames.Num(); )
+			{
+				FString NodeName = AllNodeNames[NodeIndex];
+				auto It = ObjectBindingMap.CreateConstIterator();
+				if (It)
+				{
+					ImportFBXNode(NodeName, CurveAPI, MovieScene, *Sequencer, It.Key());
+
+					UE_LOG(LogMovieScene, Warning, TEXT("Fbx Import: Failed to find any matching node for (%s). Defaulting to first available (%s)."), *NodeName, *It.Value());
+					ObjectBindingMap.Remove(It.Key());
+					AllNodeNames.RemoveAt(NodeIndex);
+					continue;
+				}
+
+				++NodeIndex;
+			}
+		}
+
 		for (FString NodeName : AllNodeNames)
 		{
-			ImportFBXNode(NodeName, CurveAPI, MovieScene, *Sequencer, ObjectBindingMap, bMatchByNameOnly);
+			UE_LOG(LogMovieScene, Warning, TEXT("Fbx Import: Failed to find any matching node for (%s)."), *NodeName);
 		}
 
 		Sequencer->NotifyMovieSceneDataChanged(EMovieSceneDataChangeType::MovieSceneStructureItemAdded);
diff --git a/Engine/Source/Editor/MovieSceneTools/Private/TrackEditors/SkeletalAnimationTrackEditor.cpp b/Engine/Source/Editor/MovieSceneTools/Private/TrackEditors/SkeletalAnimationTrackEditor.cpp
--- a/Engine/Source/Editor/MovieSceneTools/Private/TrackEditors/SkeletalAnimationTrackEditor.cpp
+++ b/Engine/Source/Editor/MovieSceneTools/Private/TrackEditors/SkeletalAnimationTrackEditor.cpp
@@ -40,6 +40,9 @@
 #include "Framework/Notifications/NotificationManager.h"
 #include "Widgets/Notifications/SNotificationList.h"
 #include "Toolkits/AssetEditorManager.h"
+#include "Engine/SCS_Node.h"
+#include "Engine/SimpleConstructionScript.h"
+#include "Engine/Blueprint.h"
 
 namespace SkeletalAnimationEditorConstants
 {
@@ -74,6 +77,18 @@ USkeletalMeshComponent* AcquireSkeletalMeshFromObjectGuid(const FGuid& Guid, TSh
 	return nullptr;
 }
 
+USkeleton* GetSkeletonFromComponent(UActorComponent* InComponent)
+{
+	USkeletalMeshComponent* SkeletalMeshComp = Cast<USkeletalMeshComponent>(InComponent);
+	if (SkeletalMeshComp && SkeletalMeshComp->SkeletalMesh && SkeletalMeshComp->SkeletalMesh->Skeleton)
+	{
+		// @todo Multiple actors, multiple components
+		return SkeletalMeshComp->SkeletalMesh->Skeleton;
+	}
+
+	return nullptr;
+}
+
 USkeleton* AcquireSkeletonFromObjectGuid(const FGuid& Guid, TSharedPtr<ISequencer> SequencerPtr)
 {
 	UObject* BoundObject = SequencerPtr.IsValid() ? SequencerPtr->FindSpawnedObjectOrTemplate(Guid) : nullptr;
@@ -82,19 +97,46 @@ USkeleton* AcquireSkeletonFromObjectGuid(const FGuid& Guid, TSharedPtr<ISequence
 	{
 		for (UActorComponent* Component : Actor->GetComponents())
 		{
-			USkeletalMeshComponent* SkeletalMeshComp = Cast<USkeletalMeshComponent>(Component);
-			if (SkeletalMeshComp && SkeletalMeshComp->SkeletalMesh && SkeletalMeshComp->SkeletalMesh->Skeleton)
+			if (USkeleton* Skeleton = GetSkeletonFromComponent(Component))
 			{
-				// @todo Multiple actors, multiple components
-				return SkeletalMeshComp->SkeletalMesh->Skeleton;
+				return Skeleton;
+			}
+		}
+
+		AActor* ActorCDO = Cast<AActor>(Actor->GetClass()->GetDefaultObject());
+		if (ActorCDO)
+		{
+			for (UActorComponent* Component : ActorCDO->GetComponents())
+			{
+				if (USkeleton* Skeleton = GetSkeletonFromComponent(Component))
+				{
+					return Skeleton;
+				}
+			}
+		}
+
+		UBlueprintGeneratedClass* ActorBlueprintGeneratedClass = Cast<UBlueprintGeneratedClass>(Actor->GetClass());
+		if (ActorBlueprintGeneratedClass)
+		{
+			const TArray<USCS_Node*>& ActorBlueprintNodes = ActorBlueprintGeneratedClass->SimpleConstructionScript->GetAllNodes();
+
+			for (USCS_Node* Node : ActorBlueprintNodes)
+			{
+				if (Node->ComponentClass->IsChildOf(USkeletalMeshComponent::StaticClass()))
+				{
+					if (USkeleton* Skeleton = GetSkeletonFromComponent(Node->GetActualComponentTemplate(ActorBlueprintGeneratedClass)))
+					{
+						return Skeleton;
+					}
+				}
 			}
 		}
 	}
 	else if(USkeletalMeshComponent* SkeletalMeshComponent = Cast<USkeletalMeshComponent>(BoundObject))
 	{
-		if (SkeletalMeshComponent->SkeletalMesh)
+		if (USkeleton* Skeleton = GetSkeletonFromComponent(SkeletalMeshComponent))
 		{
-			return SkeletalMeshComponent->SkeletalMesh->Skeleton;
+			return Skeleton;
 		}
 	}
 
diff --git a/Engine/Source/Editor/MovieSceneTools/Public/MatineeImportTools.h b/Engine/Source/Editor/MovieSceneTools/Public/MatineeImportTools.h
--- a/Engine/Source/Editor/MovieSceneTools/Public/MatineeImportTools.h
+++ b/Engine/Source/Editor/MovieSceneTools/Public/MatineeImportTools.h
@@ -11,6 +11,8 @@ class IMovieScenePlayer;
 class UInterpTrackAnimControl;
 class UInterpTrackBoolProp;
 class UInterpTrackColorProp;
+class UInterpTrackFloatMaterialParam;
+class UInterpTrackVectorMaterialParam;
 class UInterpTrackDirector;
 class UInterpTrackEvent;
 class UInterpTrackFade;
@@ -25,6 +27,7 @@ class UMovieSceneAudioTrack;
 class UMovieSceneBoolTrack;
 class UMovieSceneCameraCutTrack;
 class UMovieSceneColorTrack;
+class UMovieSceneComponentMaterialTrack;
 class UMovieSceneEventTrack;
 class UMovieSceneFadeTrack;
 class UMovieSceneFloatTrack;
@@ -56,6 +59,12 @@ public:
 	/** Copies keys from a matinee float track to a sequencer float track. */
 	static bool CopyInterpFloatTrack( UInterpTrackFloatBase* MatineeFloatTrack, UMovieSceneFloatTrack* FloatTrack );
 
+	/** Copies keys from a matinee material paramter track to a sequencer material track*/
+	static bool CopyInterpMaterialParamTrack( UInterpTrackFloatMaterialParam* MatineeMaterialParamTrack, UMovieSceneComponentMaterialTrack* MaterialTrack );
+
+	/** Copies keys from a matinee material paramter track to a sequencer material track */
+	static bool CopyInterpMaterialParamTrack( UInterpTrackVectorMaterialParam* MatineeMaterialParamTrack, UMovieSceneComponentMaterialTrack* MaterialTrack);
+
 	/** Copies keys from a matinee vector track to a sequencer vector track. */
 	static bool CopyInterpVectorTrack( UInterpTrackVectorProp* MatineeVectorTrack, UMovieSceneVectorTrack* VectorTrack );
 
diff --git a/Engine/Source/Editor/Sequencer/Private/DisplayNodes/SequencerObjectBindingNode.cpp b/Engine/Source/Editor/Sequencer/Private/DisplayNodes/SequencerObjectBindingNode.cpp
--- a/Engine/Source/Editor/Sequencer/Private/DisplayNodes/SequencerObjectBindingNode.cpp
+++ b/Engine/Source/Editor/Sequencer/Private/DisplayNodes/SequencerObjectBindingNode.cpp
@@ -488,11 +488,23 @@ FLinearColor FSequencerObjectBindingNode::GetDisplayNameColor() const
 
 	// Spawnables don't have valid object bindings when their track hasn't spawned them yet,
 	// so we override the default behavior of red with a gray so that users don't think there is something wrong.
-	if (GetBindingType() == EObjectBindingType::Spawnable)
+	
+	TSharedPtr<FSequencerDisplayNode> CurrentNode = SharedThis((FSequencerDisplayNode*)this);
+
+	while (CurrentNode.IsValid())
 	{
-		return FLinearColor::Gray;
+		if (CurrentNode->GetType() == ESequencerNode::Object)
+		{
+			if (StaticCastSharedPtr<FSequencerObjectBindingNode>(CurrentNode)->GetBindingType() == EObjectBindingType::Spawnable)
+			{
+				return FLinearColor::Gray;
+			}
+		}
+
+		CurrentNode = CurrentNode->GetParent();
 	}
 
+
 	return FLinearColor::Red;
 }
 
@@ -1133,7 +1145,9 @@ void FSequencerObjectBindingNode::HandlePropertyMenuItemExecute(FPropertyPath Pr
 		}
 	}
 
-	FKeyPropertyParams KeyPropertyParams(KeyableBoundObjects, PropertyPath, ESequencerKeyMode::ManualKey);
+	// When auto setting track defaults are disabled, force add a key so that the changed
+	// value is saved and is propagated to the property.
+	FKeyPropertyParams KeyPropertyParams(KeyableBoundObjects, PropertyPath, Sequencer.GetAutoSetTrackDefaults() == false ? ESequencerKeyMode::ManualKeyForced : ESequencerKeyMode::ManualKey);
 
 	Sequencer.KeyProperty(KeyPropertyParams);
 }
diff --git a/Engine/Source/Editor/Sequencer/Private/LevelEditorSequencerIntegration.cpp b/Engine/Source/Editor/Sequencer/Private/LevelEditorSequencerIntegration.cpp
--- a/Engine/Source/Editor/Sequencer/Private/LevelEditorSequencerIntegration.cpp
+++ b/Engine/Source/Editor/Sequencer/Private/LevelEditorSequencerIntegration.cpp
@@ -41,6 +41,7 @@
 #include "Widgets/Text/STextBlock.h"
 #include "UObject/ObjectKey.h"
 #include "UnrealEdGlobals.h"
+#include "UnrealEdMisc.h"
 #include "Editor/UnrealEdEngine.h"
 #include "EditorSupportDelegates.h"
 
@@ -238,13 +239,15 @@ void FLevelEditorSequencerIntegration::Initialize()
 	{
 		FLevelEditorModule& LevelEditorModule = FModuleManager::Get().GetModuleChecked<FLevelEditorModule>("LevelEditor");
 
-		FDelegateHandle Handle = LevelEditorModule.OnTabContentChanged().AddRaw(this, &FLevelEditorSequencerIntegration::OnTabContentChanged);
+		FDelegateHandle TabContentChanged = LevelEditorModule.OnTabContentChanged().AddRaw(this, &FLevelEditorSequencerIntegration::OnTabContentChanged);
+		FDelegateHandle MapChanged = LevelEditorModule.OnMapChanged().AddRaw(this, &FLevelEditorSequencerIntegration::OnMapChanged);
 		AcquiredResources.Add(
 			[=]{
 				FLevelEditorModule* LevelEditorModulePtr = FModuleManager::Get().GetModulePtr<FLevelEditorModule>("LevelEditor");
 				if (LevelEditorModulePtr)
 				{
-					LevelEditorModulePtr->OnTabContentChanged().Remove(Handle);
+					LevelEditorModulePtr->OnTabContentChanged().Remove(TabContentChanged);
+					LevelEditorModulePtr->OnMapChanged().Remove(MapChanged);
 				}
 			}
 		);
@@ -963,6 +966,29 @@ void FLevelEditorSequencerIntegration::OnTabContentChanged()
 
 }
 
+void FLevelEditorSequencerIntegration::OnMapChanged(UWorld* World, EMapChangeType MapChangeType)
+{
+	if (MapChangeType == EMapChangeType::TearDownWorld)
+	{
+		IterateAllSequencers(
+			[=](FSequencer& In, const FLevelEditorSequencerIntegrationOptions& Options)
+		{
+			if (Options.bRequiresLevelEvents)
+			{
+				In.GetEvaluationTemplate().ResetDirectorInstances();
+				In.RestorePreAnimatedState();
+				In.State.ClearObjectCaches(In);
+
+				// Notify data changed to enqueue an evaluate
+				In.NotifyMovieSceneDataChanged(EMovieSceneDataChangeType::Unknown);
+			}
+		}
+		);
+	}
+}
+
+
+
 void FLevelEditorSequencerIntegration::AddSequencer(TSharedRef<ISequencer> InSequencer, const FLevelEditorSequencerIntegrationOptions& Options)
 {
 	if (!BoundSequencers.Num())
diff --git a/Engine/Source/Editor/Sequencer/Private/MovieSceneSequenceEditor.cpp b/Engine/Source/Editor/Sequencer/Private/MovieSceneSequenceEditor.cpp
--- a/Engine/Source/Editor/Sequencer/Private/MovieSceneSequenceEditor.cpp
+++ b/Engine/Source/Editor/Sequencer/Private/MovieSceneSequenceEditor.cpp
@@ -31,7 +31,7 @@ namespace
 		FText DefaultComment_08 = LOCTEXT("DefaultComment_08", "\t\t\ta) Objects bound to the track's object binding, or:");
 		FText DefaultComment_09 = LOCTEXT("DefaultComment_09", "\t\t\tb) Objects specified on the track's event receivers array, or:");
 		FText DefaultComment_10 = LOCTEXT("DefaultComment_10", "\t\t\tc) Objects provided by the playback context (level blueprints, widgets etc)");
-		FText DefaultComment_11 = LOCTEXT("DefaultComment_11", "Tip: Trigger events on level blueprints by implementing an interface on it");
+		FText DefaultComment_11 = LOCTEXT("DefaultComment_11", "Tip: Trigger events on level blueprints using an interface pin, and adding the interface to the desired level blueprint");
 
 		return FText::Format(LOCTEXT("DefaultComment_Format", "{0}\n{1}\n{2}\n{3}\n{4}\n{5}\n{6}\n{7}\n{8}\n{9}\n\n{10}"),
 			DefaultComment_01, DefaultComment_02, DefaultComment_03, DefaultComment_04, DefaultComment_05,
@@ -108,7 +108,14 @@ UK2Node_FunctionEntry* FMovieSceneSequenceEditor::CreateEventEndpoint(UMovieScen
 
 void FMovieSceneSequenceEditor::InitializeEndpointForTrack(UMovieSceneEventTrack* EventTrack, UK2Node_FunctionEntry* Endpoint) const
 {
-	SetupDefaultPinForEndpoint(EventTrack, Endpoint);
+	UMovieScene* MovieScene = EventTrack->GetTypedOuter<UMovieScene>();
+
+	// Only add an object pin for the event if this is an event track on an object binding
+	FGuid Unused;
+	if (MovieScene->FindTrackBinding(*EventTrack, Unused))
+	{
+		SetupDefaultPinForEndpoint(EventTrack, Endpoint);
+	}
 }
 
 UClass* FMovieSceneSequenceEditor::FindTrackObjectBindingClass(UMovieSceneTrack* Track)
diff --git a/Engine/Source/Editor/Sequencer/Private/SSequencer.cpp b/Engine/Source/Editor/Sequencer/Private/SSequencer.cpp
--- a/Engine/Source/Editor/Sequencer/Private/SSequencer.cpp
+++ b/Engine/Source/Editor/Sequencer/Private/SSequencer.cpp
@@ -1098,8 +1098,6 @@ void SSequencer::FillTimeDisplayFormatMenu(FMenuBuilder& MenuBuilder)
 
 	if (Settings)
 	{
-		EFrameNumberDisplayFormats CurrentDisplay = Settings->GetTimeDisplayFormat();
-
 		for (int32 Index = 0; Index < FrameNumberDisplayEnum->NumEnums() - 1; Index++)
 		{
 			if (!FrameNumberDisplayEnum->HasMetaData(TEXT("Hidden"), Index))
@@ -1110,7 +1108,6 @@ void SSequencer::FillTimeDisplayFormatMenu(FMenuBuilder& MenuBuilder)
 				if (Value == EFrameNumberDisplayFormats::DropFrameTimecode && !bSupportsDropFormatDisplay)
 					continue;
 
-				bool bIsSet = CurrentDisplay == Value;
 				MenuBuilder.AddMenuEntry(
 					FrameNumberDisplayEnum->GetDisplayNameTextByIndex(Index),
 					FrameNumberDisplayEnum->GetToolTipTextByIndex(Index),
@@ -1118,7 +1115,7 @@ void SSequencer::FillTimeDisplayFormatMenu(FMenuBuilder& MenuBuilder)
 					FUIAction(
 						FExecuteAction::CreateUObject(Settings, &USequencerSettings::SetTimeDisplayFormat, Value),
 						FCanExecuteAction(),
-						FIsActionChecked::CreateLambda([=] { return CurrentDisplay == Value; })
+						FIsActionChecked::CreateLambda([=] { return Settings->GetTimeDisplayFormat() == Value; })
 					),
 					NAME_None,
 					EUserInterfaceActionType::RadioButton
@@ -1242,8 +1239,13 @@ TSharedRef<SWidget> SSequencer::MakePlaybackMenu()
 
 		MenuBuilder.AddMenuEntry( FSequencerCommands::Get().ToggleKeepCursorInPlaybackRangeWhileScrubbing );
 		MenuBuilder.AddMenuEntry( FSequencerCommands::Get().ToggleKeepCursorInPlaybackRange );
-		MenuBuilder.AddMenuEntry( FSequencerCommands::Get().ToggleKeepPlaybackRangeInSectionBounds );
-		MenuBuilder.AddMenuEntry( FSequencerCommands::Get().ToggleLinkCurveEditorTimeRange );
+
+		if (!SequencerPtr.Pin()->IsLevelEditorSequencer())
+		{
+			MenuBuilder.AddMenuEntry(FSequencerCommands::Get().ToggleKeepPlaybackRangeInSectionBounds);
+		}
+		
+		MenuBuilder.AddMenuEntry(FSequencerCommands::Get().ToggleLinkCurveEditorTimeRange);
 
 		// Menu entry for zero padding
 		auto OnZeroPadChanged = [=](uint8 NewValue){
@@ -2078,7 +2080,7 @@ void SSequencer::StepToKey(bool bStepToNextKey, bool bCameraOnly)
 				SequencerHelpers::GetAllKeyAreas( Node, KeyAreas );
 				for ( TSharedPtr<IKeyArea> KeyArea : KeyAreas )
 				{
-					KeyArea->GetKeyTimes(AllTimes);
+					KeyArea->GetKeyTimes(AllTimes, KeyArea->GetOwningSection()->GetRange());
 				}
 
 				TSet<TWeakObjectPtr<UMovieSceneSection> > Sections;
diff --git a/Engine/Source/Editor/Sequencer/Private/SSequencerPlayRateCombo.cpp b/Engine/Source/Editor/Sequencer/Private/SSequencerPlayRateCombo.cpp
--- a/Engine/Source/Editor/Sequencer/Private/SSequencerPlayRateCombo.cpp
+++ b/Engine/Source/Editor/Sequencer/Private/SSequencerPlayRateCombo.cpp
@@ -265,15 +265,12 @@ void SSequencerPlayRateCombo::PopulateClockSourceMenu(FMenuBuilder& MenuBuilder)
 
 	if (RootSequence)
 	{
-		EUpdateClockSource CurrentDisplay = RootSequence->GetMovieScene()->GetClockSource();
-
 		for (int32 Index = 0; Index < ClockSourceEnum->NumEnums() - 1; Index++)
 		{
 			if (!ClockSourceEnum->HasMetaData(TEXT("Hidden"), Index))
 			{
 				EUpdateClockSource Value = (EUpdateClockSource)ClockSourceEnum->GetValueByIndex(Index);
 
-				bool bIsSet = CurrentDisplay == Value;
 				MenuBuilder.AddMenuEntry(
 					ClockSourceEnum->GetDisplayNameTextByIndex(Index),
 					ClockSourceEnum->GetToolTipTextByIndex(Index),
@@ -281,7 +278,7 @@ void SSequencerPlayRateCombo::PopulateClockSourceMenu(FMenuBuilder& MenuBuilder)
 					FUIAction(
 						FExecuteAction::CreateSP(this, &SSequencerPlayRateCombo::SetClockSource, Value),
 						FCanExecuteAction(),
-						FIsActionChecked::CreateLambda([=]{ return CurrentDisplay == Value; })
+						FIsActionChecked::CreateLambda([=]{ return RootSequence->GetMovieScene()->GetClockSource() == Value; })
 					),
 					NAME_None,
 					EUserInterfaceActionType::RadioButton
diff --git a/Engine/Source/Editor/Sequencer/Private/Sequencer.cpp b/Engine/Source/Editor/Sequencer/Private/Sequencer.cpp
--- a/Engine/Source/Editor/Sequencer/Private/Sequencer.cpp
+++ b/Engine/Source/Editor/Sequencer/Private/Sequencer.cpp
@@ -131,6 +131,16 @@
 
 DEFINE_LOG_CATEGORY(LogSequencer);
 
+static TAutoConsoleVariable<float> CVarAutoScrubSpeed(
+	TEXT("Sequencer.AutoScrubSpeed"),
+	6.0f,
+	TEXT("How fast to scrub forward/backward when auto-scrubbing"));
+
+static TAutoConsoleVariable<float> CVarAutoScrubCurveExponent(
+	TEXT("Sequencer.AutoScrubCurveExponent"),
+	2.0f,
+	TEXT("How much to ramp in and out the scrub speed when auto-scrubbing"));
+
 struct FSequencerTemplateStore : IMovieSceneSequenceTemplateStore
 {
 	void Reset()
@@ -602,6 +612,35 @@ void FSequencer::Tick(float InDeltaTime)
 		PlayPosition.Reset(GlobalTime.ConvertTo(PlayPosition.GetInputRate()));
 	}
 
+	if (AutoScrubTarget.IsSet())
+	{
+		const double ScrubSpeed = CVarAutoScrubSpeed->GetFloat();		// How fast to scrub at peak curve speed
+		const double AutoScrubExp = CVarAutoScrubCurveExponent->GetFloat();	// How long to ease in and out.  Bigger numbers allow for longer easing.
+
+		const double SecondsPerFrame = GetFocusedTickResolution().AsInterval() / ScrubSpeed;
+		const int32 TotalFrames = FMath::Abs(AutoScrubTarget.GetValue().DestinationTime.GetFrame().Value - AutoScrubTarget.GetValue().SourceTime.GetFrame().Value);
+		const double TargetSeconds = (double)TotalFrames * SecondsPerFrame;
+
+		double ElapsedSeconds = FPlatformTime::Seconds() - AutoScrubTarget.GetValue().StartTime;
+		float Alpha = ElapsedSeconds / TargetSeconds;
+		Alpha = FMath::Clamp(Alpha, 0.f, 1.f);
+		int32 NewFrameNumber = FMath::InterpEaseInOut(AutoScrubTarget.GetValue().SourceTime.GetFrame().Value, AutoScrubTarget.GetValue().DestinationTime.GetFrame().Value, Alpha, AutoScrubExp);
+
+		FAutoScrubTarget CachedTarget = AutoScrubTarget.GetValue();
+
+		SetPlaybackStatus(EMovieScenePlayerStatus::Scrubbing);
+		PlayPosition.SetTimeBase(GetRootTickResolution(), GetRootTickResolution(), EMovieSceneEvaluationType::WithSubFrames);
+		SetLocalTimeDirectly(FFrameNumber(NewFrameNumber));
+
+		AutoScrubTarget = CachedTarget;
+
+		if (FMath::IsNearlyEqual(Alpha, 1.f, KINDA_SMALL_NUMBER))
+		{
+			SetPlaybackStatus(EMovieScenePlayerStatus::Stopped);
+			AutoScrubTarget.Reset();
+		}
+	}
+
 	UpdateSubSequenceData();
 
 	// Tick all the tools we own as well
@@ -2214,7 +2253,7 @@ void FSequencer::SetLocalTimeDirectly(FFrameTime NewTime)
 }
 
 
-void FSequencer::SetGlobalTime( FFrameTime NewTime )
+void FSequencer::SetGlobalTime(FFrameTime NewTime)
 {
 	NewTime = ConvertFrameTime(NewTime, GetRootTickResolution(), PlayPosition.GetInputRate());
 	if (PlayPosition.GetEvaluationType() == EMovieSceneEvaluationType::FrameLocked)
@@ -2229,6 +2268,12 @@ void FSequencer::SetGlobalTime( FFrameTime NewTime )
 	{
 		EvaluateInternal(PlayPosition.JumpTo(NewTime));
 	}
+
+	if (AutoScrubTarget.IsSet())
+	{
+		SetPlaybackStatus(EMovieScenePlayerStatus::Stopped);
+		AutoScrubTarget.Reset();
+	}
 }
 
 void FSequencer::ForceEvaluate()
@@ -2379,6 +2424,11 @@ TOptional<float> FSequencer::CalculateAutoscrollEncroachment(double NewTime, flo
 }
 
 
+void FSequencer::AutoScrubToTime(FFrameTime DestinationTime)
+{
+	AutoScrubTarget = FAutoScrubTarget(DestinationTime, GetLocalTime().Time, FPlatformTime::Seconds());
+}
+
 void FSequencer::SetPerspectiveViewportPossessionEnabled(bool bEnabled)
 {
 	bPerspectiveViewportPossessionEnabled = bEnabled;
@@ -2613,6 +2663,16 @@ FGuid FSequencer::GetHandleToObject( UObject* Object, bool bCreateHandleIfMissin
 
 		ObjectGuid = CreateBinding(*Object, PossessedActor != nullptr ? PossessedActor->GetActorLabel() : Object->GetName());
 
+		AActor* ActorAdded = PossessedActor ? PossessedActor : Object->GetTypedOuter<AActor>();
+		if (ActorAdded)
+		{
+			FGuid ActorAddedGuid = GetHandleToObject(ActorAdded);
+			if (ActorAddedGuid.IsValid())
+			{
+				OnActorAddedToSequencerEvent.Broadcast(ActorAdded, ActorAddedGuid);
+			}
+		}
+
 		NotifyMovieSceneDataChanged( EMovieSceneDataChangeType::MovieSceneStructureItemAdded );
 	}
 	
@@ -3824,7 +3884,14 @@ void FSequencer::OnScrubPositionChanged( FFrameTime NewScrubPosition, bool bScru
 		NewScrubPosition = FMath::Clamp(NewScrubPosition, LowerBound, UpperBound);
 	}
 
-	SetLocalTimeDirectly( NewScrubPosition );
+	if (!bScrubbing && FSlateApplication::Get().GetModifierKeys().IsShiftDown())
+	{
+		AutoScrubToTime(NewScrubPosition);
+	}
+	else
+	{
+		SetLocalTimeDirectly(NewScrubPosition);
+	}
 }
 
 
@@ -3972,7 +4039,7 @@ FGuid FSequencer::AddSpawnable(UObject& Object, UActorFactory* ActorFactory)
 	return NewGuid;
 }
 
-FGuid FSequencer::MakeNewSpawnable( UObject& Object, UActorFactory* ActorFactory )
+FGuid FSequencer::MakeNewSpawnable( UObject& Object, UActorFactory* ActorFactory, bool bSetupDefaults )
 {
 	UMovieSceneSequence* Sequence = GetFocusedMovieSceneSequence();
 	UMovieScene* MovieScene = Sequence->GetMovieScene();
@@ -4013,8 +4080,11 @@ FGuid FSequencer::MakeNewSpawnable( UObject& Object, UActorFactory* ActorFactory
 	// Spawn the object so we can position it correctly, it's going to get spawned anyway since things default to spawned.
 	UObject* SpawnedObject = SpawnRegister->SpawnObject(NewGuid, *MovieScene, ActiveTemplateIDs.Top(), *this);
 
-	FTransformData TransformData;
-	SpawnRegister->SetupDefaultsForSpawnable(SpawnedObject, Spawnable->GetGuid(), TransformData, AsShared(), Settings);
+	if (bSetupDefaults)
+	{
+		FTransformData TransformData;
+		SpawnRegister->SetupDefaultsForSpawnable(SpawnedObject, Spawnable->GetGuid(), TransformData, AsShared(), Settings);
+	}
 
 	Spawnable->SetSpawnOwnership(SavedOwnership);
 
@@ -6184,7 +6254,7 @@ bool FSequencer::PasteObjectBindings(const FString& TextToImport)
 			// because we need to use our binding (which has tracks associated with it). To solve this, we let it create
 			// an object template based off of our (transient package owned) template, then find the newly created binding
 			// and update it.
-			FGuid NewGuid = MakeNewSpawnable(*CopyableBinding->SpawnableObjectTemplate);
+			FGuid NewGuid = MakeNewSpawnable(*CopyableBinding->SpawnableObjectTemplate, nullptr, false);
 			FMovieSceneBinding NewBinding(NewGuid, CopyableBinding->Binding.GetName(), CopyableBinding->Tracks);
 			FMovieSceneSpawnable* Spawnable = MovieScene->FindSpawnable(NewGuid);
 
@@ -6900,9 +6970,28 @@ FMovieSceneSpawnable* FSequencer::ConvertToSpawnableInternal(FGuid PossessableGu
 		// Remap all the spawnable's tracks and child bindings onto the new possessable
 		MovieScene->MoveBindingContents(PossessableGuid, SpawnableGuid);
 
+		FMovieSceneBinding* PossessableBinding = (FMovieSceneBinding*)MovieScene->GetBindings().FindByPredicate([&](FMovieSceneBinding& Binding) { return Binding.GetObjectGuid() == PossessableGuid; });
+		check(PossessableBinding);
+
+		for (UMovieSceneFolder* Folder : MovieScene->GetRootFolders())
+		{
+			if (ReplaceFolderBindingGUID(Folder, PossessableGuid, SpawnableGuid))
+			{
+				break;
+			}
+		}
+
+		int32 SortingOrder = PossessableBinding->GetSortingOrder();
+
 		if (MovieScene->RemovePossessable(PossessableGuid))
 		{
 			Sequence->UnbindPossessableObjects(PossessableGuid);
+
+			FMovieSceneBinding* SpawnableBinding = (FMovieSceneBinding*)MovieScene->GetBindings().FindByPredicate([&](FMovieSceneBinding& Binding) { return Binding.GetObjectGuid() == SpawnableGuid; });
+			check(SpawnableBinding);
+			
+			SpawnableBinding->SetSortingOrder(SortingOrder);
+
 		}
 
 		TOptional<FTransformData> TransformData;
@@ -7078,12 +7167,30 @@ FMovieScenePossessable* FSequencer::ConvertToPossessableInternal(FGuid Spawnable
 		// Remap all the spawnable's tracks and child bindings onto the new possessable
 		MovieScene->MoveBindingContents(OldSpawnableGuid, NewPossessableGuid);
 
+		FMovieSceneBinding* SpawnableBinding = (FMovieSceneBinding*)MovieScene->GetBindings().FindByPredicate([&](FMovieSceneBinding& Binding) { return Binding.GetObjectGuid() == OldSpawnableGuid; });
+		check(SpawnableBinding);
+
+		for (UMovieSceneFolder* Folder : MovieScene->GetRootFolders())
+		{
+			if (ReplaceFolderBindingGUID(Folder, Spawnable->GetGuid(), Possessable->GetGuid()))
+			{
+				break;
+			}
+		}
+
+		int32 SortingOrder = SpawnableBinding->GetSortingOrder();
+
 		// Remove the spawnable and all it's sub tracks
 		if (MovieScene->RemoveSpawnable(OldSpawnableGuid))
 		{
 			SpawnRegister->DestroySpawnedObject(OldSpawnableGuid, ActiveTemplateIDs.Top(), *this);
+
+			FMovieSceneBinding* PossessableBinding = (FMovieSceneBinding*)MovieScene->GetBindings().FindByPredicate([&](FMovieSceneBinding& Binding) { return Binding.GetObjectGuid() == NewPossessableGuid; });
+			check(PossessableBinding);
+			
+			PossessableBinding->SetSortingOrder(SortingOrder);
 		}
-	
+
 		static const FName SequencerActorTag(TEXT("SequencerActor"));
 		PossessedActor->Tags.Remove(SequencerActorTag);
 
@@ -7095,6 +7202,36 @@ FMovieScenePossessable* FSequencer::ConvertToPossessableInternal(FGuid Spawnable
 	return Possessable;
 }
 
+bool FSequencer::ReplaceFolderBindingGUID(UMovieSceneFolder* Folder, FGuid Original, FGuid Converted)
+{
+	UMovieScene* MovieScene = GetFocusedMovieSceneSequence()->GetMovieScene();
+
+	if (MovieScene->IsReadOnly())
+	{
+		return true;
+	}
+
+	for (FGuid ChildGuid : Folder->GetChildObjectBindings())
+	{
+		if (ChildGuid == Original)
+		{
+			Folder->AddChildObjectBinding(Converted);
+			Folder->RemoveChildObjectBinding(Original);
+			return true;
+		}
+	}
+
+	for (UMovieSceneFolder* ChildFolder : Folder->GetChildFolders())
+	{
+		if (ReplaceFolderBindingGUID(ChildFolder, Original, Converted))
+		{
+			return true;
+		}
+	}
+
+	return false;
+}
+
 void FSequencer::OnAddFolder()
 {
 	UMovieScene* FocusedMovieScene = GetFocusedMovieSceneSequence()->GetMovieScene();
@@ -7698,7 +7835,7 @@ void FSequencer::StepToNextMark()
 				{
 					if (TickFrame > CurrentTickFrameNumber)
 					{
-						SetLocalTime(TickFrame.Value);
+						AutoScrubToTime(TickFrame.Value);
 						break;
 					}
 				}
@@ -7729,7 +7866,7 @@ void FSequencer::StepToPreviousMark()
 				{
 					if (TickFrame < CurrentTickFrameNumber)
 					{
-						SetLocalTime(TickFrame.Value);
+						AutoScrubToTime(TickFrame.Value);
 						break;
 					}
 				}
diff --git a/Engine/Source/Editor/Sequencer/Private/Sequencer.h b/Engine/Source/Editor/Sequencer/Private/Sequencer.h
--- a/Engine/Source/Editor/Sequencer/Private/Sequencer.h
+++ b/Engine/Source/Editor/Sequencer/Private/Sequencer.h
@@ -618,6 +618,9 @@ public:
 	 */
 	void UpdateAutoScroll(double NewTime);
 
+	/** Autoscrub to destination time */
+	void AutoScrubToTime(FFrameTime DestinationTime);
+
 public:
 
 	//~ FGCObject Interface
@@ -723,7 +726,7 @@ public:
 	virtual void Pause() override;
 	virtual TSharedRef<SWidget> MakeTimeRange(const TSharedRef<SWidget>& InnerContent, bool bShowWorkingRange, bool bShowViewRange, bool bShowPlaybackRange) override;
 	virtual UObject* FindSpawnedObjectOrTemplate(const FGuid& BindingId) override;
-	virtual FGuid MakeNewSpawnable(UObject& SourceObject, UActorFactory* ActorFactory = nullptr) override;
+	virtual FGuid MakeNewSpawnable(UObject& SourceObject, UActorFactory* ActorFactory = nullptr, bool bSetupDefaults = true) override;
 	virtual bool IsReadOnly() const override;
 	virtual void ExternalSelectionHasChanged() override { SynchronizeSequencerSelectionWithExternalSelection(); }
 	/** Access the user-supplied settings object */
@@ -920,6 +923,9 @@ protected:
 	/** Internal conversion function that doesn't perform expensive reset/update tasks */
 	FMovieScenePossessable* ConvertToPossessableInternal(FGuid SpawnableGuid);
 
+	/** Recurses through a folder to replace converted GUID with new GUID */
+	bool ReplaceFolderBindingGUID(UMovieSceneFolder *Folder, FGuid Original, FGuid Converted);
+
 	/** Internal function to render movie for a given start/end time */
 	void RenderMovieInternal(TRange<FFrameNumber> Range, bool bSetFrameOverrides = false) const;
 
@@ -1049,6 +1055,19 @@ private:
 	/** The amount of autoscrub offset that is currently being applied */
 	TOptional<float> AutoscrubOffset;
 
+	struct FAutoScrubTarget
+	{
+		FAutoScrubTarget(FFrameTime InDestinationTime, FFrameTime InSourceTime, double InStartTime)
+			: DestinationTime(InDestinationTime)
+			, SourceTime(InSourceTime)
+			, StartTime(InStartTime) {}
+		FFrameTime DestinationTime;
+		FFrameTime SourceTime;
+		double StartTime;
+	};
+
+	TOptional<FAutoScrubTarget> AutoScrubTarget;
+
 	/** Zoom smoothing curves */
 	FCurveSequence ZoomAnimation;
 	FCurveHandle ZoomCurve;
diff --git a/Engine/Source/Editor/Sequencer/Private/SequencerKeyCollection.cpp b/Engine/Source/Editor/Sequencer/Private/SequencerKeyCollection.cpp
--- a/Engine/Source/Editor/Sequencer/Private/SequencerKeyCollection.cpp
+++ b/Engine/Source/Editor/Sequencer/Private/SequencerKeyCollection.cpp
@@ -1,6 +1,7 @@
 // Copyright 1998-2018 Epic Games, Inc. All Rights Reserved.
 
 #include "SequencerKeyCollection.h"
+#include "MovieSceneSection.h"
 #include "IKeyArea.h"
 #include "DisplayNodes/SequencerDisplayNode.h"
 #include "DisplayNodes/SequencerTrackNode.h"
@@ -160,7 +161,7 @@ bool FSequencerKeyCollection::Update(const FSequencerKeyCollectionSignature& InS
 	// Get all the key times for the key areas
 	for (auto& Pair : InSignature.GetKeyAreas())
 	{
-		Pair.Key->GetKeyTimes(AllTimes);
+		Pair.Key->GetKeyTimes(AllTimes, Pair.Key->GetOwningSection()->GetRange());
 	}
 
 	AllTimes.Sort();
diff --git a/Engine/Source/Editor/Sequencer/Public/ISequencer.h b/Engine/Source/Editor/Sequencer/Public/ISequencer.h
--- a/Engine/Source/Editor/Sequencer/Public/ISequencer.h
+++ b/Engine/Source/Editor/Sequencer/Public/ISequencer.h
@@ -226,7 +226,7 @@ public:
 	 * @param	ActorFactory	Optional actor factory to use to create spawnable type
 	 * @return	The spawnable guid for the spawnable, or an invalid Guid if we were not able to create a spawnable
 	 */
-	virtual FGuid MakeNewSpawnable(UObject& SourceObject, UActorFactory* ActorFactory = nullptr) = 0;
+	virtual FGuid MakeNewSpawnable(UObject& SourceObject, UActorFactory* ActorFactory = nullptr, bool bSetupDefaults = false) = 0;
 
 	/**
 	 * Add actors as possessable objects to sequencer.
diff --git a/Engine/Source/Editor/Sequencer/Public/LevelEditorSequencerIntegration.h b/Engine/Source/Editor/Sequencer/Public/LevelEditorSequencerIntegration.h
--- a/Engine/Source/Editor/Sequencer/Public/LevelEditorSequencerIntegration.h
+++ b/Engine/Source/Editor/Sequencer/Public/LevelEditorSequencerIntegration.h
@@ -107,6 +107,9 @@ private:
 	/** Called when viewport tab content changes. */
 	void OnTabContentChanged();
 
+	/** Called when the map is changed. */
+	void OnMapChanged(UWorld* World, EMapChangeType MapChangeType);
+
 	/** Called before a PIE session begins. */
 	void OnPreBeginPIE(bool bIsSimulating);
 
diff --git a/Engine/Source/Editor/UnrealEd/Private/Fbx/FbxMainExport.cpp b/Engine/Source/Editor/UnrealEd/Private/Fbx/FbxMainExport.cpp
--- a/Engine/Source/Editor/UnrealEd/Private/Fbx/FbxMainExport.cpp
+++ b/Engine/Source/Editor/UnrealEd/Private/Fbx/FbxMainExport.cpp
@@ -588,6 +588,8 @@ void FFbxExporter::FillFbxCameraAttribute(FbxNode* ParentNode, FbxCamera* Camera
 	float ApertureHeightInInches = 0.612f; // 0.612f is a magic number from Maya that represents the ApertureHeight
 	float ApertureWidthInInches = CameraComponent->AspectRatio * ApertureHeightInInches;
 	float FocalLength = Camera->ComputeFocalLength(CameraComponent->FieldOfView);
+	
+	TOptional<float> FocusDistance;
 
 	if (CameraComponent->IsA(UCineCameraComponent::StaticClass()))
 	{
@@ -597,6 +599,7 @@ void FFbxExporter::FillFbxCameraAttribute(FbxNode* ParentNode, FbxCamera* Camera
 			ApertureWidthInInches = FUnitConversion::Convert(CineCameraComponent->FilmbackSettings.SensorWidth, EUnit::Millimeters, EUnit::Inches);
 			ApertureHeightInInches = FUnitConversion::Convert(CineCameraComponent->FilmbackSettings.SensorHeight, EUnit::Millimeters, EUnit::Inches);
 			FocalLength = CineCameraComponent->CurrentFocalLength;
+			FocusDistance = CineCameraComponent->FocusSettings.ManualFocusDistance;
 		}
 	}
 
@@ -609,6 +612,11 @@ void FFbxExporter::FillFbxCameraAttribute(FbxNode* ParentNode, FbxCamera* Camera
 	Camera->SetApertureMode(FbxCamera::eFocalLength);
 	Camera->FocalLength.Set(FocalLength);
 
+	if (FocusDistance.IsSet())
+	{
+		Camera->FocusDistance.Set(FocusDistance.GetValue());
+	}
+
 	// Add one user property for recording the AspectRatio animation
 	CreateAnimatableUserProperty(ParentNode, CameraComponent->AspectRatio, "UE_AspectRatio", "UE_Matinee_Camera_AspectRatio");
 
diff --git a/Engine/Source/Runtime/AnimGraphRuntime/Private/AnimCustomInstance.cpp b/Engine/Source/Runtime/AnimGraphRuntime/Private/AnimCustomInstance.cpp
--- a/Engine/Source/Runtime/AnimGraphRuntime/Private/AnimCustomInstance.cpp
+++ b/Engine/Source/Runtime/AnimGraphRuntime/Private/AnimCustomInstance.cpp
@@ -46,7 +46,7 @@ void UAnimCustomInstance::UnbindFromSkeletalMeshComponent(USkeletalMeshComponent
 	}
 
 	// if not game world, don't clean this up
-	if (InSkeletalMeshComponent->GetWorld()->IsGameWorld() == false)
+	if (InSkeletalMeshComponent->GetWorld() != nullptr && InSkeletalMeshComponent->GetWorld()->IsGameWorld() == false)
 	{
 		InSkeletalMeshComponent->ClearMotionVector();
 	}
diff --git a/Engine/Source/Runtime/Core/Public/UObject/SequencerObjectVersion.h b/Engine/Source/Runtime/Core/Public/UObject/SequencerObjectVersion.h
--- a/Engine/Source/Runtime/Core/Public/UObject/SequencerObjectVersion.h
+++ b/Engine/Source/Runtime/Core/Public/UObject/SequencerObjectVersion.h
@@ -39,6 +39,9 @@ struct CORE_API FSequencerObjectVersion
 		// Finish UMG evaluation on end
 		FinishUMGEvaluation,
 
+		// Manual serialization of float channel
+		SerializeFloatChannel,
+
 		// -----<new versions can be added above this line>-------------------------------------------------
 		VersionPlusOne,
 		LatestVersion = VersionPlusOne - 1
diff --git a/Engine/Source/Runtime/Engine/Classes/Animation/AnimMontage.h b/Engine/Source/Runtime/Engine/Classes/Animation/AnimMontage.h
--- a/Engine/Source/Runtime/Engine/Classes/Animation/AnimMontage.h
+++ b/Engine/Source/Runtime/Engine/Classes/Animation/AnimMontage.h
@@ -522,8 +522,8 @@ public:
 	ENGINE_API static UAnimMontage* SetMatineeAnimPositionInner(FName SlotName, USkeletalMeshComponent* SkeletalMeshComponent, UAnimSequenceBase* InAnimSequence, float InPosition, bool bLooping);
 	ENGINE_API static UAnimMontage* PreviewMatineeSetAnimPositionInner(FName SlotName, USkeletalMeshComponent* SkeletalMeshComponent, UAnimSequenceBase* InAnimSequence, float InPosition, bool bLooping, bool bFireNotifies, float DeltaTime);
 	/** static functions that are used by sequencer montage support*/
-	ENGINE_API static UAnimMontage* SetSequencerMontagePosition(FName SlotName, USkeletalMeshComponent* SkeletalMeshComponent, int32& InOutInstanceId, UAnimSequenceBase* InAnimSequence, float InPosition, float Weight, bool bLooping);
-	ENGINE_API static UAnimMontage* PreviewSequencerMontagePosition(FName SlotName, USkeletalMeshComponent* SkeletalMeshComponent, int32& InOutInstanceId, UAnimSequenceBase* InAnimSequence, float InPosition, float Weight, bool bLooping, bool bFireNotifies);
+	ENGINE_API static UAnimMontage* SetSequencerMontagePosition(FName SlotName, USkeletalMeshComponent* SkeletalMeshComponent, int32& InOutInstanceId, UAnimSequenceBase* InAnimSequence, float InPosition, float Weight, bool bLooping, bool bPlaying);
+	ENGINE_API static UAnimMontage* PreviewSequencerMontagePosition(FName SlotName, USkeletalMeshComponent* SkeletalMeshComponent, int32& InOutInstanceId, UAnimSequenceBase* InAnimSequence, float InPosition, float Weight, bool bLooping, bool bFireNotifies, bool bPlaying);
 private:
 	static UAnimMontage* InitializeMatineeControl(FName SlotName, USkeletalMeshComponent* SkeletalMeshComponent, UAnimSequenceBase* InAnimSequence, bool bLooping);
 };
diff --git a/Engine/Source/Runtime/Engine/Classes/Matinee/InterpTrackFloatMaterialParam.h b/Engine/Source/Runtime/Engine/Classes/Matinee/InterpTrackFloatMaterialParam.h
--- a/Engine/Source/Runtime/Engine/Classes/Matinee/InterpTrackFloatMaterialParam.h
+++ b/Engine/Source/Runtime/Engine/Classes/Matinee/InterpTrackFloatMaterialParam.h
@@ -10,7 +10,7 @@
 class UInterpTrackInst;
 struct FPropertyChangedEvent;
 
-UCLASS(meta=( DisplayName = "Float Material Parameter Track" ) )
+UCLASS(MinimalAPI, meta=( DisplayName = "Float Material Parameter Track" ) )
 class UInterpTrackFloatMaterialParam : public UInterpTrackFloatBase
 {
 	GENERATED_UCLASS_BODY()
diff --git a/Engine/Source/Runtime/Engine/Private/Animation/AnimMontage.cpp b/Engine/Source/Runtime/Engine/Private/Animation/AnimMontage.cpp
--- a/Engine/Source/Runtime/Engine/Private/Animation/AnimMontage.cpp
+++ b/Engine/Source/Runtime/Engine/Private/Animation/AnimMontage.cpp
@@ -2671,7 +2671,7 @@ UAnimMontage* FAnimMontageInstance::PreviewMatineeSetAnimPositionInner(FName Slo
 	return PlayingMontage;
 }
 
-UAnimMontage* FAnimMontageInstance::SetSequencerMontagePosition(FName SlotName, USkeletalMeshComponent* SkeletalMeshComponent, int32& InOutInstanceId, UAnimSequenceBase* InAnimSequence, float InPosition, float Weight, bool bLooping)
+UAnimMontage* FAnimMontageInstance::SetSequencerMontagePosition(FName SlotName, USkeletalMeshComponent* SkeletalMeshComponent, int32& InOutInstanceId, UAnimSequenceBase* InAnimSequence, float InPosition, float Weight, bool bLooping, bool bPlaying)
 {
 	UAnimInstance* AnimInst = SkeletalMeshComponent->GetAnimInstance();
 	if (AnimInst)
@@ -2699,6 +2699,7 @@ UAnimMontage* FAnimMontageInstance::SetSequencerMontagePosition(FName SlotName,
 			MontageInstanceToUpdate->Blend.SetDesiredValue(Weight);
 			MontageInstanceToUpdate->Blend.SetAlpha(Weight);
 			MontageInstanceToUpdate->SetPosition(InPosition);
+			MontageInstanceToUpdate->bPlaying = bPlaying;
 			return PlayingMontage;
 		}
 	}
@@ -2710,7 +2711,7 @@ UAnimMontage* FAnimMontageInstance::SetSequencerMontagePosition(FName SlotName,
 	return nullptr;
 }
 
-UAnimMontage* FAnimMontageInstance::PreviewSequencerMontagePosition(FName SlotName, USkeletalMeshComponent* SkeletalMeshComponent, int32& InOutInstanceId, UAnimSequenceBase* InAnimSequence, float InPosition, float Weight, bool bLooping, bool bFireNotifies)
+UAnimMontage* FAnimMontageInstance::PreviewSequencerMontagePosition(FName SlotName, USkeletalMeshComponent* SkeletalMeshComponent, int32& InOutInstanceId, UAnimSequenceBase* InAnimSequence, float InPosition, float Weight, bool bLooping, bool bFireNotifies, bool bPlaying)
 {
 	UAnimInstance* AnimInst = SkeletalMeshComponent->GetAnimInstance();
 	if (AnimInst)
@@ -2718,7 +2719,7 @@ UAnimMontage* FAnimMontageInstance::PreviewSequencerMontagePosition(FName SlotNa
 		FAnimMontageInstance* MontageInstanceToUpdate = AnimInst->GetMontageInstanceForID(InOutInstanceId);
 		float PreviousPosition = (MontageInstanceToUpdate) ? MontageInstanceToUpdate->GetPosition() : InPosition;
 
-		UAnimMontage* PlayingMontage = SetSequencerMontagePosition(SlotName, SkeletalMeshComponent, InOutInstanceId, InAnimSequence, InPosition, Weight, bLooping);
+		UAnimMontage* PlayingMontage = SetSequencerMontagePosition(SlotName, SkeletalMeshComponent, InOutInstanceId, InAnimSequence, InPosition, Weight, bLooping, bPlaying);
 		if (PlayingMontage)
 		{
 			// we have to get it again in case if this is new
diff --git a/Engine/Source/Runtime/LevelSequence/Private/LevelSequenceActor.cpp b/Engine/Source/Runtime/LevelSequence/Private/LevelSequenceActor.cpp
--- a/Engine/Source/Runtime/LevelSequence/Private/LevelSequenceActor.cpp
+++ b/Engine/Source/Runtime/LevelSequence/Private/LevelSequenceActor.cpp
@@ -57,7 +57,7 @@ ALevelSequenceActor::ALevelSequenceActor(const FObjectInitializer& Init)
 	bOverrideInstanceData = false;
 
 	PrimaryActorTick.bCanEverTick = true;
-	bAutoPlay = false;
+	bAutoPlay_DEPRECATED = false;
 }
 
 bool ALevelSequenceActor::RetrieveBindingOverrides(const FGuid& InBindingId, FMovieSceneSequenceID InSequenceID, TArray<UObject*, TInlineAllocator<1>>& OutObjects) const
@@ -93,7 +93,7 @@ void ALevelSequenceActor::BeginPlay()
 	InitializePlayer();
 	RefreshBurnIn();
 
-	if (bAutoPlay)
+	if (PlaybackSettings.bAutoPlay)
 	{
 		SequencePlayer->Play();
 	}
@@ -113,6 +113,13 @@ void ALevelSequenceActor::PostLoad()
 {
 	Super::PostLoad();
 
+	// If autoplay was previously enabled, initialize the playback settings to autoplay
+	if (bAutoPlay_DEPRECATED)
+	{
+		PlaybackSettings.bAutoPlay = bAutoPlay_DEPRECATED;
+		bAutoPlay_DEPRECATED = false;
+	}
+
 	// ~~~~~~~~~~~~~~~~~~~~~~~~~~~~~~~~~~~~~~~~~~~~~~~~~~~~~~~~~~~~~~~~~~~~~~~~~~~~~~~~~~~~~~~~~~~~~~~~~~~
 	// We intentionally do not attempt to load any asset in PostLoad other than by way of LoadPackageAsync
 	// since under some circumstances it is possible for the sequence to only be partially loaded.
diff --git a/Engine/Source/Runtime/LevelSequence/Private/LevelSequenceDirector.cpp b/Engine/Source/Runtime/LevelSequence/Private/LevelSequenceDirector.cpp
--- a/Engine/Source/Runtime/LevelSequence/Private/LevelSequenceDirector.cpp
+++ b/Engine/Source/Runtime/LevelSequence/Private/LevelSequenceDirector.cpp
@@ -5,5 +5,9 @@
 
 UWorld* ULevelSequenceDirector::GetWorld() const
 {
+	if (ULevel* OuterLevel = GetTypedOuter<ULevel>())
+	{
+		return OuterLevel->OwningWorld;
+	}
 	return GetTypedOuter<UWorld>();
 }
diff --git a/Engine/Source/Runtime/LevelSequence/Private/LevelSequencePlayer.cpp b/Engine/Source/Runtime/LevelSequence/Private/LevelSequencePlayer.cpp
--- a/Engine/Source/Runtime/LevelSequence/Private/LevelSequencePlayer.cpp
+++ b/Engine/Source/Runtime/LevelSequence/Private/LevelSequencePlayer.cpp
@@ -52,6 +52,10 @@ ULevelSequencePlayer* ULevelSequencePlayer::CreateLevelSequencePlayer(UObject* W
 	SpawnParams.SpawnCollisionHandlingOverride = ESpawnActorCollisionHandlingMethod::AlwaysSpawn;
 	SpawnParams.ObjectFlags |= RF_Transient;
 	SpawnParams.bAllowDuringConstructionScript = true;
+
+	// Defer construction for autoplay so that BeginPlay() is called
+	SpawnParams.bDeferConstruction = true;
+
 	ALevelSequenceActor* Actor = World->SpawnActor<ALevelSequenceActor>(SpawnParams);
 
 	Actor->PlaybackSettings = Settings;
@@ -60,6 +64,9 @@ ULevelSequencePlayer* ULevelSequencePlayer::CreateLevelSequencePlayer(UObject* W
 	Actor->InitializePlayer();
 	OutActor = Actor;
 
+	FTransform DefaultTransform;
+	Actor->FinishSpawning(DefaultTransform);
+
 	return Actor->SequencePlayer;
 }
 
@@ -108,6 +115,20 @@ void ULevelSequencePlayer::OnStopped()
 			Actor->PrimaryActorTick.RemovePrerequisite(LevelSequenceActor, LevelSequenceActor->PrimaryActorTick);
 		}
 	}
+
+	if (World != nullptr && World->GetGameInstance() != nullptr)
+	{
+		APlayerController* PC = World->GetGameInstance()->GetFirstLocalPlayerController();
+
+		if (PC != nullptr)
+		{
+			if (PC->PlayerCameraManager)
+			{
+				PC->PlayerCameraManager->bClientSimulatingViewTarget = false;
+			}
+		}
+	}
+
 	PrerequisiteActors.Reset();
 	LastViewTarget.Reset();
 }
@@ -184,6 +205,12 @@ void ULevelSequencePlayer::UpdateCameraCut(UObject* CameraObject, UObject* Unloc
 	if (CameraActor == nullptr)
 	{
 		CameraActor = LastViewTarget.Get();
+
+		// Skip if the last view target is the same as the current view target so that there's no additional camera cut
+		if (CameraActor == ViewTarget)
+		{
+			return;
+		}
 	}
 
 	FViewTargetTransitionParams TransitionParams;
diff --git a/Engine/Source/Runtime/LevelSequence/Public/LevelSequenceActor.h b/Engine/Source/Runtime/LevelSequence/Public/LevelSequenceActor.h
--- a/Engine/Source/Runtime/LevelSequence/Public/LevelSequenceActor.h
+++ b/Engine/Source/Runtime/LevelSequence/Public/LevelSequenceActor.h
@@ -74,8 +74,8 @@ public:
 
 public:
 
-	UPROPERTY(EditAnywhere, BlueprintReadOnly, Category="Playback")
-	uint32 bAutoPlay : 1;
+	UPROPERTY()
+	bool bAutoPlay_DEPRECATED;
 
 	UPROPERTY(EditAnywhere, BlueprintReadOnly, Category="Playback", meta=(ShowOnlyInnerProperties))
 	FMovieSceneSequencePlaybackSettings PlaybackSettings;
diff --git a/Engine/Source/Runtime/MovieScene/Private/Channels/MovieSceneFloatChannel.cpp b/Engine/Source/Runtime/MovieScene/Private/Channels/MovieSceneFloatChannel.cpp
--- a/Engine/Source/Runtime/MovieScene/Private/Channels/MovieSceneFloatChannel.cpp
+++ b/Engine/Source/Runtime/MovieScene/Private/Channels/MovieSceneFloatChannel.cpp
@@ -4,7 +4,64 @@
 #include "MovieSceneFwd.h"
 #include "Channels/MovieSceneChannelProxy.h"
 #include "MovieSceneFrameMigration.h"
+#include "UObject/SequencerObjectVersion.h"
 
+bool FMovieSceneTangentData::Serialize(FArchive& Ar)
+{
+	Ar.UsingCustomVersion(FSequencerObjectVersion::GUID);
+	if (Ar.CustomVer(FSequencerObjectVersion::GUID) < FSequencerObjectVersion::SerializeFloatChannel)
+	{
+		return false;
+	}
+
+	// Serialization is handled manually to avoid the extra size overhead of UProperty tagging.
+	// Otherwise with many keys in a FMovieSceneTangentData the size can become quite large.
+	Ar << ArriveTangent;
+	Ar << LeaveTangent;
+	Ar << TangentWeightMode;
+	Ar << ArriveTangentWeight;
+	Ar << LeaveTangentWeight;
+
+	return true;
+}
+
+bool FMovieSceneTangentData::operator==(const FMovieSceneTangentData& TangentData) const
+{
+	return (ArriveTangent == TangentData.ArriveTangent) && (LeaveTangent == TangentData.LeaveTangent) && (TangentWeightMode == TangentData.TangentWeightMode) && (ArriveTangentWeight == TangentData.ArriveTangentWeight) && (LeaveTangentWeight == TangentData.LeaveTangentWeight);
+}
+
+bool FMovieSceneTangentData::operator!=(const FMovieSceneTangentData& Other) const
+{
+	return !(*this == Other);
+}
+
+bool FMovieSceneFloatValue::Serialize(FArchive& Ar)
+{
+	Ar.UsingCustomVersion(FSequencerObjectVersion::GUID);
+	if (Ar.CustomVer(FSequencerObjectVersion::GUID) < FSequencerObjectVersion::SerializeFloatChannel)
+	{
+		return false;
+	}
+
+	// Serialization is handled manually to avoid the extra size overhead of UProperty tagging.
+	// Otherwise with many keys in a FMovieSceneFloatValue the size can become quite large.
+	Ar << Value;
+	Ar << InterpMode;
+	Ar << TangentMode;
+	Ar << Tangent;
+
+	return true;
+}
+
+bool FMovieSceneFloatValue::operator==(const FMovieSceneFloatValue& FloatValue) const
+{
+	return (Value == FloatValue.Value) && (InterpMode == FloatValue.InterpMode) && (TangentMode == FloatValue.TangentMode) && (Tangent == FloatValue.Tangent);
+}
+
+bool FMovieSceneFloatValue::operator!=(const FMovieSceneFloatValue& Other) const
+{
+	return !(*this == Other);
+}
 
 bool FMovieSceneFloatChannel::SerializeFromMismatchedTag(const FPropertyTag& Tag, FStructuredArchive::FSlot Slot)
 {
diff --git a/Engine/Source/Runtime/MovieScene/Private/Compilation/MovieSceneCompiler.cpp b/Engine/Source/Runtime/MovieScene/Private/Compilation/MovieSceneCompiler.cpp
--- a/Engine/Source/Runtime/MovieScene/Private/Compilation/MovieSceneCompiler.cpp
+++ b/Engine/Source/Runtime/MovieScene/Private/Compilation/MovieSceneCompiler.cpp
@@ -314,6 +314,7 @@ void FMovieSceneCompiler::CompileRange(TRange<FFrameNumber> InGlobalRange, UMovi
 		// next entry in the evaluation field iterator (which should be a populated range)
 		if (CompiledRange.GetUpperBound() == EmptySpaceRange.GetUpperBound())
 		{
+			IterFromBound = TRangeBound<FFrameNumber>::FlipInclusion(CompiledRange.GetUpperBound());
 			++ExistingEvaluationFieldIter;
 		}
 	}
diff --git a/Engine/Source/Runtime/MovieScene/Private/Compilation/MovieSceneSegmentCompiler.cpp b/Engine/Source/Runtime/MovieScene/Private/Compilation/MovieSceneSegmentCompiler.cpp
--- a/Engine/Source/Runtime/MovieScene/Private/Compilation/MovieSceneSegmentCompiler.cpp
+++ b/Engine/Source/Runtime/MovieScene/Private/Compilation/MovieSceneSegmentCompiler.cpp
@@ -4,27 +4,6 @@
 #include "Compilation/MovieSceneCompilerRules.h"
 #include "Evaluation/MovieSceneEvaluationTree.h"
 
-PRAGMA_DISABLE_DEPRECATION_WARNINGS
-void FMovieSceneSegmentCompilerRules::Blend(FSegmentBlendData& BlendData) const
-{
-	FMovieSceneSegment TmpSegment;
-
-	for (int32 Index = 0; Index < BlendData.Num(); ++Index)
-	{
-		TmpSegment.Impls.Add(FSectionEvaluationData(Index, BlendData[Index].Flags));
-	}
-
-	BlendSegment(TmpSegment, BlendData);
-
-	FSegmentBlendData NewBlendData;
-	for (FSectionEvaluationData EvalData : TmpSegment.Impls)
-	{
-		NewBlendData.Add(BlendData[EvalData.ImplIndex]);
-		NewBlendData.Last().Flags = EvalData.Flags;
-	}
-}
-PRAGMA_ENABLE_DEPRECATION_WARNINGS
-
 bool FMovieSceneAdditiveCameraTrackBlender::SortByStartTime(const FMovieSceneSectionData& A, const FMovieSceneSectionData& B)
 {
 	return TRangeBound<FFrameNumber>::MinLower(A.Section->GetRange().GetLowerBound(), B.Section->GetRange().GetLowerBound()) == A.Section->GetRange().GetLowerBound();
diff --git a/Engine/Source/Runtime/MovieScene/Private/MovieSceneSequencePlayer.cpp b/Engine/Source/Runtime/MovieScene/Private/MovieSceneSequencePlayer.cpp
--- a/Engine/Source/Runtime/MovieScene/Private/MovieSceneSequencePlayer.cpp
+++ b/Engine/Source/Runtime/MovieScene/Private/MovieSceneSequencePlayer.cpp
@@ -261,20 +261,29 @@ void UMovieSceneSequencePlayer::Scrub()
 }
 
 void UMovieSceneSequencePlayer::Stop()
+{
+	StopInternal(bReversePlayback ? GetLastValidTime() : FFrameTime(StartTime));
+}
+
+void UMovieSceneSequencePlayer::StopAtCurrentTime()
+{
+	StopInternal(PlayPosition.GetCurrentPosition());
+}
+
+void UMovieSceneSequencePlayer::StopInternal(FFrameTime TimeToResetTo)
 {
 	if (IsPlaying() || IsPaused() || RootTemplateInstance.IsValid())
 	{
 		if (bIsEvaluating)
 		{
-			LatentActions.Emplace(FLatentAction::Stop);
+			LatentActions.Emplace(FLatentAction::Stop, TimeToResetTo);
 			return;
 		}
 
 		Status = EMovieScenePlayerStatus::Stopped;
 
-		// Put the cursor back to the start
-		PlayPosition.Reset(bReversePlayback ? GetLastValidTime() : FFrameTime(StartTime));
-
+		// Put the cursor at the specified position
+		PlayPosition.Reset(TimeToResetTo);
 		TimeController->StopPlaying(GetCurrentTime());
 
 		CurrentNumLoops = 0;
@@ -294,7 +303,7 @@ void UMovieSceneSequencePlayer::Stop()
 		if (HasAuthority())
 		{
 			// Explicitly handle Stop() events through an RPC call
-			RPC_OnStopEvent(PlayPosition.GetCurrentPosition());
+			RPC_OnStopEvent(TimeToResetTo);
 		}
 		UpdateNetworkSyncProperties();
 
@@ -315,8 +324,9 @@ void UMovieSceneSequencePlayer::Stop()
 
 void UMovieSceneSequencePlayer::GoToEndAndStop()
 {
-	JumpToFrame(GetLastValidTime());
-	Stop();
+	FFrameTime Time = GetLastValidTime();
+	JumpToFrame(Time);
+	StopInternal(Time);
 }
 
 FQualifiedFrameTime UMovieSceneSequencePlayer::GetCurrentTime() const
@@ -515,7 +525,7 @@ void UMovieSceneSequencePlayer::Initialize(UMovieSceneSequence* InSequence, cons
 	// may be called during PostLoad.
 	if (Sequence)
 	{
-		Stop();
+		StopAtCurrentTime();
 	}
 
 	Sequence = InSequence;
@@ -714,10 +724,7 @@ void UMovieSceneSequencePlayer::UpdateTimeCursorPosition_Internal(FFrameTime New
 			}
 			else
 			{
-				Stop();
-
-				// When playback stops naturally, the time cursor is put at the boundary that was crossed to make ping-pong playback easy
-				PlayPosition.Reset(NewPosition);
+				StopInternal(NewPosition);
 			}
 
 			TimeController->StopPlaying(GetCurrentTime());
@@ -776,8 +783,8 @@ void UMovieSceneSequencePlayer::ApplyLatentActions()
 	{
 		switch (LatentAction.Type)
 		{
-		case FLatentAction::Stop:          Stop();                         continue;
-		case FLatentAction::Pause:         Pause();                        continue;
+		case FLatentAction::Stop:          StopInternal(LatentAction.Position); continue;
+		case FLatentAction::Pause:         Pause();                             continue;
 		}
 
 		check(LatentAction.Type == FLatentAction::Update);
@@ -906,9 +913,9 @@ void UMovieSceneSequencePlayer::RPC_OnStopEvent_Implementation(FFrameTime Stoppe
 	// ~~~~~~~~~~~~~~~~~~~~~~~~~~~~~~~~~~~~~~~~~~~~~~~~
 	// Handle an explicit Stop command from the server.
 
-	if (HasAuthority())
+	if (HasAuthority() || !Sequence)
 	{
-		// Never run network sync operations on authoritative players
+		// Never run network sync operations on authoritative players or players that have not been initialized yet
 		return;
 	}
 
@@ -937,7 +944,7 @@ void UMovieSceneSequencePlayer::RPC_OnStopEvent_Implementation(FFrameTime Stoppe
 	case EMovieScenePlayerStatus::Scrubbing: ScrubToFrame(StoppedTime);  break;
 	}
 
-	Stop();
+	StopInternal(StoppedTime);
 }
 
 void UMovieSceneSequencePlayer::PostNetReceive()
@@ -947,9 +954,9 @@ void UMovieSceneSequencePlayer::PostNetReceive()
 
 	Super::PostNetReceive();
 
-	if (!ensure(!HasAuthority()))
+	if (!ensure(!HasAuthority()) || !Sequence)
 	{
-		// Never run network sync operations on authoritative players
+		// Never run network sync operations on authoritative players or players that have not been initialized yet
 		return;
 	}
 
diff --git a/Engine/Source/Runtime/MovieScene/Private/MovieSceneTrack.cpp b/Engine/Source/Runtime/MovieScene/Private/MovieSceneTrack.cpp
--- a/Engine/Source/Runtime/MovieScene/Private/MovieSceneTrack.cpp
+++ b/Engine/Source/Runtime/MovieScene/Private/MovieSceneTrack.cpp
@@ -138,40 +138,14 @@ void UMovieSceneTrack::UpdateEasing()
 	}
 }
 
-PRAGMA_DISABLE_DEPRECATION_WARNINGS
-TInlineValue<FMovieSceneSegmentCompilerRules> UMovieSceneTrack::GetRowCompilerRules() const
-{
-	return TInlineValue<FMovieSceneSegmentCompilerRules>();
-}
-
-TInlineValue<FMovieSceneSegmentCompilerRules> UMovieSceneTrack::GetTrackCompilerRules() const
-{
-	return TInlineValue<FMovieSceneSegmentCompilerRules>();
-}
-
 FMovieSceneTrackRowSegmentBlenderPtr UMovieSceneTrack::GetRowSegmentBlender() const
 {
-	// Handle legacy row compiler rules
-	TInlineValue<FMovieSceneSegmentCompilerRules> LegacyRules = GetRowCompilerRules();
-	if (LegacyRules.IsValid())
-	{
-		return TLegacyTrackRowSegmentBlender<FMovieSceneTrackRowSegmentBlender>(MoveTemp(LegacyRules));
-	}
-	else
-	{
-		return FDefaultTrackRowSegmentBlender();
-	}
+	return FDefaultTrackRowSegmentBlender();
 }
 
 FMovieSceneTrackSegmentBlenderPtr UMovieSceneTrack::GetTrackSegmentBlender() const
 {
-	// Handle legacy track compiler rules
-	TInlineValue<FMovieSceneSegmentCompilerRules> LegacyRules = GetTrackCompilerRules();
-	if (LegacyRules.IsValid())
-	{
-		return TLegacyTrackRowSegmentBlender<FMovieSceneTrackSegmentBlender>(MoveTemp(LegacyRules));
-	}
-	else if (EvalOptions.bCanEvaluateNearestSection && EvalOptions.bEvalNearestSection)
+	if (EvalOptions.bCanEvaluateNearestSection && EvalOptions.bEvalNearestSection)
 	{
 		return FEvaluateNearestSegmentBlender();
 	}
@@ -180,7 +154,6 @@ FMovieSceneTrackSegmentBlenderPtr UMovieSceneTrack::GetTrackSegmentBlender() con
 		return FMovieSceneTrackSegmentBlenderPtr();
 	}
 }
-PRAGMA_ENABLE_DEPRECATION_WARNINGS
 
 void UMovieSceneTrack::GenerateTemplate(const FMovieSceneTrackCompilerArgs& Args) const
 {
diff --git a/Engine/Source/Runtime/MovieScene/Public/Channels/MovieSceneFloatChannel.h b/Engine/Source/Runtime/MovieScene/Public/Channels/MovieSceneFloatChannel.h
--- a/Engine/Source/Runtime/MovieScene/Public/Channels/MovieSceneFloatChannel.h
+++ b/Engine/Source/Runtime/MovieScene/Public/Channels/MovieSceneFloatChannel.h
@@ -13,7 +13,6 @@
 
 #include "MovieSceneFloatChannel.generated.h"
 
-
 USTRUCT()
 struct FMovieSceneTangentData
 {
@@ -27,6 +26,15 @@ struct FMovieSceneTangentData
 		, LeaveTangentWeight(0.f)
 	{}
 
+	bool Serialize(FArchive& Ar);
+	bool operator==(const FMovieSceneTangentData& Other) const;
+	bool operator!=(const FMovieSceneTangentData& Other) const;
+	friend FArchive& operator<<(FArchive& Ar, FMovieSceneTangentData& P)
+	{
+		P.Serialize(Ar);
+		return Ar;
+	}
+
 	/** If RCIM_Cubic, the arriving tangent at this key */
 	UPROPERTY(EditAnywhere, Category="Key")
 	float ArriveTangent;
@@ -49,6 +57,25 @@ struct FMovieSceneTangentData
 
 };
 
+template<>
+struct TIsPODType<FMovieSceneTangentData>
+{
+	enum { Value = true };
+};
+
+
+template<>
+struct TStructOpsTypeTraits<FMovieSceneTangentData>
+	: public TStructOpsTypeTraitsBase2<FMovieSceneTangentData>
+{
+	enum
+	{
+		WithSerializer = true,
+		WithCopy = false,
+		WithIdenticalViaEquality = true,
+	};
+};
+
 USTRUCT()
 struct FMovieSceneFloatValue
 {
@@ -62,6 +89,15 @@ struct FMovieSceneFloatValue
 		: Value(InValue), InterpMode(RCIM_Cubic), TangentMode(RCTM_Auto)
 	{}
 
+	bool Serialize(FArchive& Ar);
+	bool operator==(const FMovieSceneFloatValue& Other) const;
+	bool operator!=(const FMovieSceneFloatValue& Other) const;
+	friend FArchive& operator<<(FArchive& Ar, FMovieSceneFloatValue& P)
+	{
+		P.Serialize(Ar);
+		return Ar;
+	}
+
 	UPROPERTY(EditAnywhere, Category="Key")
 	float Value;
 
@@ -75,6 +111,25 @@ struct FMovieSceneFloatValue
 	FMovieSceneTangentData Tangent;
 };
 
+template<>
+struct TIsPODType<FMovieSceneFloatValue>
+{
+	enum { Value = true };
+};
+
+
+template<>
+struct TStructOpsTypeTraits<FMovieSceneFloatValue>
+	: public TStructOpsTypeTraitsBase2<FMovieSceneFloatValue>
+{
+	enum
+	{
+		WithSerializer = true,
+		WithCopy = false,
+		WithIdenticalViaEquality = true,
+	};
+};
+
 USTRUCT()
 struct MOVIESCENE_API FMovieSceneFloatChannel : public FMovieSceneChannel
 {
diff --git a/Engine/Source/Runtime/MovieScene/Public/Compilation/IMovieSceneTemplateGenerator.h b/Engine/Source/Runtime/MovieScene/Public/Compilation/IMovieSceneTemplateGenerator.h
--- a/Engine/Source/Runtime/MovieScene/Public/Compilation/IMovieSceneTemplateGenerator.h
+++ b/Engine/Source/Runtime/MovieScene/Public/Compilation/IMovieSceneTemplateGenerator.h
@@ -27,24 +27,6 @@ struct IMovieSceneTemplateGenerator
 	 */
 	virtual void AddOwnedTrack(FMovieSceneEvaluationTrack&& InTrackTemplate, const UMovieSceneTrack& SourceTrack) = 0;
 
-
-public:
-
-	UE_DEPRECATED(4.19, "Shared tracks are no longer supported, please emit shared Execution tokens instead (FMovieSceneExecutionTokens::AddShared")
-	void AddSharedTrack(FMovieSceneEvaluationTrack&& InTrackTemplate, const FMovieSceneSharedDataId& SharedId, const UMovieSceneTrack& SourceTrack){}
-
-	UE_DEPRECATED(4.19, "Support for legacy tracks has now been removed")
-	void AddLegacyTrack(FMovieSceneEvaluationTrack&& InTrackTemplate, const UMovieSceneTrack& SourceTrack){}
-
-	UE_DEPRECATED(4.19, "External segments are now handled internally")
-	void AddExternalSegments(TRange<float> RootRange, TArrayView<const FMovieSceneEvaluationFieldSegmentPtr> SegmentPtrs, ESectionEvaluationFlags Flags){}
-
-	UE_DEPRECATED(4.19, "Sub sequence templates are no longer generated at compile time within template generators")
-	FMovieSceneSequenceTransform GetSequenceTransform(FMovieSceneSequenceIDRef InSequenceID) const { return FMovieSceneSequenceTransform(); }
-
-	UE_DEPRECATED(4.19, "Sub sequences are now handled internally")
-	void AddSubSequence(const FMovieSceneSubSequenceData& SequenceData, FMovieSceneSequenceIDRef ParentID, FMovieSceneSequenceID SequenceID) {}
-
 protected:
 	virtual ~IMovieSceneTemplateGenerator() { }
 };
diff --git a/Engine/Source/Runtime/MovieScene/Public/Compilation/MovieSceneCompilerRules.h b/Engine/Source/Runtime/MovieScene/Public/Compilation/MovieSceneCompilerRules.h
--- a/Engine/Source/Runtime/MovieScene/Public/Compilation/MovieSceneCompilerRules.h
+++ b/Engine/Source/Runtime/MovieScene/Public/Compilation/MovieSceneCompilerRules.h
@@ -145,30 +145,3 @@ private:
 
 	MOVIESCENE_API static bool SortByStartTime(const FMovieSceneSectionData& A, const FMovieSceneSectionData& B);
 };
-
-
-PRAGMA_DISABLE_DEPRECATION_WARNINGS
-
-/** Track row segment blender that acts as a proxy to an instance of the legacy FMovieSceneSegmentCompilerRules */
-template<typename BaseType>
-struct TLegacyTrackRowSegmentBlender : BaseType
-{
-	TLegacyTrackRowSegmentBlender(TInlineValue<FMovieSceneSegmentCompilerRules>&& InRules)
-		: LegacyRules(MoveTemp(InRules))
-	{}
-
-	~TLegacyTrackRowSegmentBlender()
-	{}
-
-	TLegacyTrackRowSegmentBlender(TLegacyTrackRowSegmentBlender&&) = default;
-
-	virtual void Blend(FSegmentBlendData& BlendData) const override
-	{
-		return LegacyRules->Blend(BlendData);
-	}
-
-private:
-	TInlineValue<FMovieSceneSegmentCompilerRules> LegacyRules;
-};
-
-PRAGMA_ENABLE_DEPRECATION_WARNINGS
\ No newline at end of file
diff --git a/Engine/Source/Runtime/MovieScene/Public/Compilation/MovieSceneSegmentCompiler.h b/Engine/Source/Runtime/MovieScene/Public/Compilation/MovieSceneSegmentCompiler.h
--- a/Engine/Source/Runtime/MovieScene/Public/Compilation/MovieSceneSegmentCompiler.h
+++ b/Engine/Source/Runtime/MovieScene/Public/Compilation/MovieSceneSegmentCompiler.h
@@ -140,13 +140,3 @@ protected:
 
 typedef TInlineValue<FMovieSceneTrackRowSegmentBlender, 16> FMovieSceneTrackRowSegmentBlenderPtr;
 typedef TInlineValue<FMovieSceneTrackSegmentBlender, 16> FMovieSceneTrackSegmentBlenderPtr;
-
-
-/** Deprecated type support */
-struct FMovieSceneSegmentCompilerRules : FMovieSceneTrackSegmentBlender
-{
-	virtual void Blend(FSegmentBlendData& BlendData) const override;
-
-protected:
-	virtual void BlendSegment(FMovieSceneSegment& Segment, const TArrayView<const FMovieSceneSectionData>& SourceData) const = 0;
-};
diff --git a/Engine/Source/Runtime/MovieScene/Public/Evaluation/MovieSceneEvalTemplate.h b/Engine/Source/Runtime/MovieScene/Public/Evaluation/MovieSceneEvalTemplate.h
--- a/Engine/Source/Runtime/MovieScene/Public/Evaluation/MovieSceneEvalTemplate.h
+++ b/Engine/Source/Runtime/MovieScene/Public/Evaluation/MovieSceneEvalTemplate.h
@@ -113,14 +113,6 @@ public:
 	 * @param ExecutionTokens	Stack of execution tokens that will be used to apply animated state to the environment at a later time.
 	*/
 	virtual void EvaluateSwept(const FMovieSceneEvaluationOperand& Operand, const FMovieSceneContext& Context, const TRange<FFrameNumber>& SweptRange, const FPersistentEvaluationData& PersistentData, FMovieSceneExecutionTokens& ExecutionTokens) const
-	{
-		// Call deprecated one
-		PRAGMA_DISABLE_DEPRECATION_WARNINGS
-		EvaluateSwept(Operand, Context, PersistentData, ExecutionTokens);
-		PRAGMA_ENABLE_DEPRECATION_WARNINGS
-	}
-	UE_DEPRECATED(4.21, "Please override void EvaluateSwept(const FMovieSceneEvaluationOperand&, const FMovieSceneContext&, const TRange<FFrameNumber>&, const FPersistentEvaluationData&, FMovieSceneExecutionTokens&) instead.")
-	virtual void EvaluateSwept(const FMovieSceneEvaluationOperand& Operand, const FMovieSceneContext& Context, const FPersistentEvaluationData& PersistentData, FMovieSceneExecutionTokens& ExecutionTokens) const
 	{
 		ensureMsgf(false, TEXT("FMovieSceneEvalTemplate::EvaluateSwept has not been implemented. Verify that this template's evaluation track has correct evaluation method (usually set in UMovieSceneTrack::PostCompile), or implement this function."));
 	}
diff --git a/Engine/Source/Runtime/MovieScene/Public/Evaluation/MovieScenePropertyTemplate.h b/Engine/Source/Runtime/MovieScene/Public/Evaluation/MovieScenePropertyTemplate.h
--- a/Engine/Source/Runtime/MovieScene/Public/Evaluation/MovieScenePropertyTemplate.h
+++ b/Engine/Source/Runtime/MovieScene/Public/Evaluation/MovieScenePropertyTemplate.h
@@ -40,55 +40,6 @@ namespace PropertyTemplate
 		FMovieSceneAnimTypeID PropertyID;
 	};
 
-	/** The value of the object as it existed before this frame's evaluation */
-	template<typename PropertyValueType> struct
-	UE_DEPRECATED(4.17, "Precaching of property values should no longer be necessary as it was only used to pass default values to curves on evaluation. Curves should now be checked for emptyness before attempting to animate an object.")
-	TCachedValue
-	{
-		TWeakObjectPtr<UObject> WeakObject;
-		PropertyValueType Value;
-	};
-
-	PRAGMA_DISABLE_DEPRECATION_WARNINGS
-	template<typename T> struct TCachedValue_DEPRECATED : TCachedValue<T> {};
-	PRAGMA_ENABLE_DEPRECATION_WARNINGS
-
-	template<typename PropertyValueType> struct 
-	UE_DEPRECATED(4.17, "Precaching of property values should no longer be necessary as it was only used to pass default values to curves on evaluation. Curves should now be checked for emptyness before attempting to animate an object.")
-	TCachedSectionData : FSectionData
-	{
-		/** Setup the data for the current frame */
-		void SetupFrame(const FMovieSceneEvaluationOperand& Operand, IMovieScenePlayer& Player)
-		{
-			ObjectsAndValues.Reset();
-			
-			for (TWeakObjectPtr<> Object : Player.FindBoundObjects(Operand))
-			{
-				if (UObject* ObjectPtr = Object.Get())
-				{
-					PropertyBindings->CacheBinding(*ObjectPtr);
-					if (UProperty* Property = PropertyBindings->GetProperty(*ObjectPtr))
-					{
-						if (Property->GetSize() == sizeof(PropertyValueType))
-						{
-							ObjectsAndValues.Add(TCachedValue_DEPRECATED<PropertyValueType>{ ObjectPtr, PropertyBindings->GetCurrentValue<PropertyValueType>(*ObjectPtr) });
-						}
-#if !(UE_BUILD_SHIPPING || UE_BUILD_TEST)
-						else
-						{
-							FMessageLog("Sequencer").Warning()
-								->AddToken(FUObjectToken::Create(Player.GetEvaluationTemplate().GetSequence(MovieSceneSequenceID::Root)))
-								->AddToken(FTextToken::Create(FText::Format(NSLOCTEXT("MovieScene", "IncompatibleDataWarning", "Property size mismatch for property '{0}'. Expected '{1}', found '{2}'. Recreate the track with the new property type."), FText::FromString(PropertyBindings->GetPropertyPath()), FText::FromString(TNameOf<PropertyValueType>::GetName()), FText::FromString(Property->GetCPPType()))));
-						}
-#endif
-					}
-				}
-			}
-		}
-
-		TArray<TCachedValue_DEPRECATED<PropertyValueType>, TInlineAllocator<1>> ObjectsAndValues;
-	};
-
 	template<typename PropertyValueType, typename IntermediateType = PropertyValueType>
 	struct TTemporarySetterType { typedef PropertyValueType Type; };
 
@@ -164,32 +115,6 @@ namespace PropertyTemplate
 	};
 }
 
-
-template<typename PropertyValueType> struct
-UE_DEPRECATED(4.17, "Precaching of property values should no longer be necessary as it was only used to pass default values to curves on evaluation. Curves should now be checked for emptyness before attempting to animate an object.")
-TCachedPropertyTrackExecutionToken : IMovieSceneExecutionToken
-{
-	virtual void Execute(const FMovieSceneContext& Context, const FMovieSceneEvaluationOperand& Operand, FPersistentEvaluationData& PersistentData, IMovieScenePlayer& Player) override
-	{
-		using namespace PropertyTemplate;
-
-		MOVIESCENE_DETAILED_SCOPE_CYCLE_COUNTER(MovieSceneEval_PropertyTrack_TokenExecute)
-		
-		TCachedSectionData<PropertyValueType>& PropertyTrackData = PersistentData.GetSectionData<TCachedSectionData<PropertyValueType>>();
-		FTrackInstancePropertyBindings* PropertyBindings = PropertyTrackData.PropertyBindings.Get();
-
-		for (TCachedValue<PropertyValueType>& ObjectAndValue : PropertyTrackData.ObjectsAndValues)
-		{
-			if (UObject* ObjectPtr = ObjectAndValue.WeakObject.Get())
-			{
-				Player.SavePreAnimatedState(*ObjectPtr, PropertyTrackData.PropertyID, FTokenProducer<PropertyValueType>(*PropertyBindings));
-				
-				PropertyBindings->CallFunction<PropertyValueType>(*ObjectPtr, ObjectAndValue.Value);
-			}
-		}
-	}
-};
-
 /** Execution token that simple stores a value, and sets it when executed */
 template<typename PropertyValueType, typename IntermediateType = PropertyValueType>
 struct TPropertyTrackExecutionToken : IMovieSceneExecutionToken
@@ -290,22 +215,6 @@ struct FMovieScenePropertySectionData
 	{
 		PersistentData.AddSectionData<T>().Initialize(PropertyName, PropertyPath, FunctionName, NotifyFunctionName);
 	}
-
-	template<typename T>
-	UE_DEPRECATED(4.17, "Precaching of property values should no longer be necessary as it was only used to pass default values to curves on evaluation. Curves should now be checked for emptyness before attempting to animate an object.")
-	void SetupCachedTrack(FPersistentEvaluationData& PersistentData) const
-	{
-		typedef PropertyTemplate::TCachedSectionData<T> FSectionData;
-		PersistentData.AddSectionData<FSectionData>().Initialize(PropertyName, PropertyPath, FunctionName, NotifyFunctionName);
-	}
-
-	template<typename T>
-	UE_DEPRECATED(4.17, "Precaching of property values should no longer be necessary as it was only used to pass default values to curves on evaluation. Curves should now be checked for emptyness before attempting to animate an object.")
-	void SetupCachedFrame(const FMovieSceneEvaluationOperand& Operand, FPersistentEvaluationData& PersistentData, IMovieScenePlayer& Player) const
-	{
-		typedef PropertyTemplate::TCachedSectionData<T> FSectionData;
-		PersistentData.GetSectionData<FSectionData>().SetupFrame(Operand, Player);
-	}
 };
 
 USTRUCT()
diff --git a/Engine/Source/Runtime/MovieScene/Public/IMovieSceneModule.h b/Engine/Source/Runtime/MovieScene/Public/IMovieSceneModule.h
--- a/Engine/Source/Runtime/MovieScene/Public/IMovieSceneModule.h
+++ b/Engine/Source/Runtime/MovieScene/Public/IMovieSceneModule.h
@@ -17,11 +17,6 @@ struct FMovieSceneEvaluationGroupParameters
 		: EvaluationPriority(InPriority)
 	{}
 
-	UE_DEPRECATED(4.17, "Please remove bInRequiresImmediateFlush parameter")
-	FMovieSceneEvaluationGroupParameters(uint16 InPriority, bool bInRequiresImmediateFlush)
-		: EvaluationPriority(InPriority)
-	{}
-
 	/** Prioirty assigned to this group. Higher priorities are evaluated first */
 	uint16 EvaluationPriority;
 };
diff --git a/Engine/Source/Runtime/MovieScene/Public/MovieSceneCommonHelpers.h b/Engine/Source/Runtime/MovieScene/Public/MovieSceneCommonHelpers.h
--- a/Engine/Source/Runtime/MovieScene/Public/MovieSceneCommonHelpers.h
+++ b/Engine/Source/Runtime/MovieScene/Public/MovieSceneCommonHelpers.h
@@ -243,17 +243,6 @@ public:
 		return PropertyName;
 	}
 
-	template<typename ValueType>
-	UE_DEPRECATED(4.15, "Please use GetCurrentValue(const UObject&)")
-	ValueType GetCurrentValue(const UObject* Object) { check(Object) return GetCurrentValue<ValueType>(*Object); }
-	template <typename ValueType>
-	UE_DEPRECATED(4.15, "Please use CallFunction(UObject&)")
-	void CallFunction( UObject* InRuntimeObject, ValueType* PropertyValue ) { CallFunction<ValueType>(*InRuntimeObject, *PropertyValue); }
-	UE_DEPRECATED(4.15, "UpdateBindings is no longer necessary")
-	void UpdateBindings( const TArray<TWeakObjectPtr<UObject>>& InRuntimeObjects ){}
-	UE_DEPRECATED(4.15, "UpdateBindings is no longer necessary")
-	void UpdateBinding( const TWeakObjectPtr<UObject>& InRuntimeObject ) {}
-
 private:
 
 	struct FPropertyAddress
diff --git a/Engine/Source/Runtime/MovieScene/Public/MovieSceneSequencePlayer.h b/Engine/Source/Runtime/MovieScene/Public/MovieSceneSequencePlayer.h
--- a/Engine/Source/Runtime/MovieScene/Public/MovieSceneSequencePlayer.h
+++ b/Engine/Source/Runtime/MovieScene/Public/MovieSceneSequencePlayer.h
@@ -83,7 +83,8 @@ USTRUCT(BlueprintType)
 struct FMovieSceneSequencePlaybackSettings
 {
 	FMovieSceneSequencePlaybackSettings()
-		: PlayRate(1.f)
+		: bAutoPlay(false)
+		, PlayRate(1.f)
 		, StartTime(0.f)
 		, bRandomStartTime(false)
 		, bRestoreState(false)
@@ -97,6 +98,10 @@ struct FMovieSceneSequencePlaybackSettings
 
 	GENERATED_BODY()
 
+	/** Auto-play the sequence when created */
+	UPROPERTY(EditAnywhere, BlueprintReadWrite, Category = "Playback")
+	uint32 bAutoPlay : 1;
+
 	/** Number of times to loop playback. -1 for infinite, else the number of times to loop before stopping */
 	UPROPERTY(EditAnywhere, BlueprintReadWrite, Category="Playback", meta=(UIMin=1, DisplayName="Loop"))
 	FMovieSceneSequenceLoopCount LoopCount;
@@ -190,10 +195,14 @@ public:
 	UFUNCTION(BlueprintCallable, Category="Game|Cinematic")
 	void Scrub();
 
-	/** Stop playback. */
+	/** Stop playback and move the cursor to the end (or start, for reversed playback) of the sequence. */
 	UFUNCTION(BlueprintCallable, Category="Game|Cinematic")
 	void Stop();
 
+	/** Stop playback without moving the cursor. */
+	UFUNCTION(BlueprintCallable, Category="Game|Cinematic")
+	void StopAtCurrentTime();
+
 	/** Go to end and stop. */
 	UFUNCTION(BlueprintCallable, Category="Game|Cinematic", meta = (ToolTip = "Go to end of the sequence and stop. Adheres to 'When Finished' section rules."))
 	void GoToEndAndStop();
@@ -325,7 +334,8 @@ public:
 public:
 
 	/**
-	 * Play the sequence from the current time, to the specified frame position
+	 * Low-level call to set the current time of the player by evaluating from the current time to the specified time, as if the sequence is playing. 
+	 * Triggers events that lie within the evaluated range. Does not alter the persistent playback status of the player (IsPlaying).
 	 *
 	 * @param NewPosition     The new frame time to play to
 	 */
@@ -333,7 +343,8 @@ public:
 	void PlayToFrame(FFrameTime NewPosition);
 
 	/**
-	 * Scrub the sequence from the current time, to the specified frame position
+	 * Low-level call to set the current time of the player by evaluating only the specified time. Will not trigger any events. 
+	 * Does not alter the persistent playback status of the player (IsPlaying).
 	 *
 	 * @param NewPosition     The new frame time to scrub to
 	 */
@@ -341,7 +352,8 @@ public:
 	void ScrubToFrame(FFrameTime NewPosition);
 
 	/**
-	 * Jump to the specified frame position, without evaluating the sequence in between the current and desired time (as if in a paused state)
+	 * Low-level call to set the current time of the player by evaluating only the specified time, as if scrubbing the timeline. Will trigger only events that exist at the specified time. 
+	 * Does not alter the persistent playback status of the player (IsPlaying).
 	 *
 	 * @param NewPosition     The new frame time to jump to
 	 */
@@ -470,6 +482,7 @@ public:
 protected:
 
 	void PlayInternal();
+	void StopInternal(FFrameTime TimeToResetTo);
 
 	void UpdateMovieSceneInstance(FMovieSceneEvaluationRange InRange, EMovieScenePlayerStatus::Type PlayerStatus, bool bHasJumped = false);
 
@@ -576,7 +589,7 @@ protected:
 	{
 		enum EType { Stop, Pause, Update };
 
-		FLatentAction(EType InType)
+		FLatentAction(EType InType, FFrameTime DesiredTime = 0)
 			: Type(InType)
 		{}
 
diff --git a/Engine/Source/Runtime/MovieScene/Public/MovieSceneTrack.h b/Engine/Source/Runtime/MovieScene/Public/MovieSceneTrack.h
--- a/Engine/Source/Runtime/MovieScene/Public/MovieSceneTrack.h
+++ b/Engine/Source/Runtime/MovieScene/Public/MovieSceneTrack.h
@@ -18,10 +18,6 @@ struct FMovieSceneTrackSegmentBlender;
 struct FMovieSceneTrackRowSegmentBlender;
 struct IMovieSceneTemplateGenerator;
 
-PRAGMA_DISABLE_DEPRECATION_WARNINGS
-	typedef TInlineValue<FMovieSceneSegmentCompilerRules> FMovieSceneDeprecatedCompilerRulesPtr;
-PRAGMA_ENABLE_DEPRECATION_WARNINGS
-
 /** Flags used to perform cook-time optimization of movie scene data */
 enum class ECookOptimizationFlags
 {
@@ -49,16 +45,6 @@ struct FMovieSceneTrackCompilerArgs
 
 	/** The generator responsible for generating the template */
 	IMovieSceneTemplateGenerator& Generator;
-
-	struct FMovieSceneSequenceTemplateStore {};
-
-	UE_DEPRECATED(4.19, "Template store is no longer supplied as part of track compilation")
-	FMovieSceneSequenceTemplateStore SubSequenceStore;
-
-	struct FMovieSceneTrackCompilationParams { 	bool bForEditorPreview, bDuringBlueprintCompile; };
-
-	UE_DEPRECATED(4.19, "Template store is no longer supplied as part of track compilation")
-	FMovieSceneTrackCompilationParams Params;
 };
 
 /** Generic evaluation options for any track */
@@ -154,12 +140,6 @@ public:
 	 */
 	MOVIESCENE_API virtual FMovieSceneTrackSegmentBlenderPtr GetTrackSegmentBlender() const;
 
-	UE_DEPRECATED(4.19, "Please override GetRowSegmentBlender() instead.")
-	MOVIESCENE_API virtual FMovieSceneDeprecatedCompilerRulesPtr GetRowCompilerRules() const;
-
-	UE_DEPRECATED(4.19, "Please override GetTrackSegmentBlender() instead.")
-	MOVIESCENE_API virtual FMovieSceneDeprecatedCompilerRulesPtr GetTrackCompilerRules() const;
-
 	/**
 	 * Generate a template for this track
 	 *
@@ -285,14 +265,6 @@ public:
 	 */
 	virtual const TArray<UMovieSceneSection*>& GetAllSections() const PURE_VIRTUAL(UMovieSceneTrack::GetAllSections, static TArray<UMovieSceneSection*> Empty; return Empty;);
 
-	/**
-	 * Gets the section boundaries of this track.
-	 * 
-	 * @return The range of time boundaries.
-	 */
-	UE_DEPRECATED(4.20, "This function is no longer used.")
-	virtual TRange<FFrameNumber> GetSectionBoundaries() const PURE_VIRTUAL(UMovieSceneTrack::GetSectionBoundaries, return TRange<FFrameNumber>::Empty(););
-
 	/**
 	 * Checks to see if the section is in this track.
 	 *
diff --git a/Engine/Source/Runtime/MovieSceneCapture/Private/VideoCaptureProtocol.cpp b/Engine/Source/Runtime/MovieSceneCapture/Private/VideoCaptureProtocol.cpp
--- a/Engine/Source/Runtime/MovieSceneCapture/Private/VideoCaptureProtocol.cpp
+++ b/Engine/Source/Runtime/MovieSceneCapture/Private/VideoCaptureProtocol.cpp
@@ -39,7 +39,6 @@ void UVideoCaptureProtocol::ConditionallyCreateWriter()
 	Options.OutputFilename = MoveTemp(VideoFilename);
 	Options.CaptureFramerateNumerator = CaptureHost->GetCaptureFrameRate().Numerator;
 	Options.CaptureFramerateDenominator = CaptureHost->GetCaptureFrameRate().Denominator;
-	Options.CodecName = VideoCodec;
 	Options.bSynchronizeFrames = CaptureHost->GetCaptureStrategy().ShouldSynchronizeFrames();
 	Options.Width = InitSettings->DesiredSize.X;
 	Options.Height = InitSettings->DesiredSize.Y;
diff --git a/Engine/Source/Runtime/MovieSceneCapture/Public/Protocols/VideoCaptureProtocol.h b/Engine/Source/Runtime/MovieSceneCapture/Public/Protocols/VideoCaptureProtocol.h
--- a/Engine/Source/Runtime/MovieSceneCapture/Public/Protocols/VideoCaptureProtocol.h
+++ b/Engine/Source/Runtime/MovieSceneCapture/Public/Protocols/VideoCaptureProtocol.h
@@ -31,9 +31,6 @@ public:
 	UPROPERTY(config, EditAnywhere, Category=VideoSettings, meta=(ClampMin=1, ClampMax=100, EditCondition=bUseCompression))
 	float CompressionQuality;
 
-	UPROPERTY(config, EditAnywhere, Category=VideoSettings, AdvancedDisplay)
-	FString VideoCodec;
-
 public:
 	virtual bool SetupImpl() override;
 	virtual void FinalizeImpl() override;
diff --git a/Engine/Source/Runtime/MovieSceneTracks/Private/Evaluation/MovieSceneAudioTemplate.cpp b/Engine/Source/Runtime/MovieSceneTracks/Private/Evaluation/MovieSceneAudioTemplate.cpp
--- a/Engine/Source/Runtime/MovieSceneTracks/Private/Evaluation/MovieSceneAudioTemplate.cpp
+++ b/Engine/Source/Runtime/MovieSceneTracks/Private/Evaluation/MovieSceneAudioTemplate.cpp
@@ -18,6 +18,7 @@
 
 
 DECLARE_CYCLE_STAT(TEXT("Audio Track Evaluate"), MovieSceneEval_AudioTrack_Evaluate, STATGROUP_MovieSceneEval);
+DECLARE_CYCLE_STAT(TEXT("Audio Track Tear Down"), MovieSceneEval_AudioTrack_TearDown, STATGROUP_MovieSceneEval);
 DECLARE_CYCLE_STAT(TEXT("Audio Track Token Execute"), MovieSceneEval_AudioTrack_TokenExecute, STATGROUP_MovieSceneEval);
 
 
@@ -384,7 +385,7 @@ void FMovieSceneAudioSectionTemplate::Evaluate(const FMovieSceneEvaluationOperan
 
 void FMovieSceneAudioSectionTemplate::TearDown(FPersistentEvaluationData& PersistentData, IMovieScenePlayer& Player) const
 {
-	MOVIESCENE_DETAILED_SCOPE_CYCLE_COUNTER(MovieSceneEval_AudioTrack_Teardown)
+	MOVIESCENE_DETAILED_SCOPE_CYCLE_COUNTER(MovieSceneEval_AudioTrack_TearDown)
 
 	if (GEngine && GEngine->UseSound())
 	{
diff --git a/Engine/Source/Runtime/MovieSceneTracks/Private/Evaluation/MovieSceneEventTemplate.cpp b/Engine/Source/Runtime/MovieSceneTracks/Private/Evaluation/MovieSceneEventTemplate.cpp
--- a/Engine/Source/Runtime/MovieSceneTracks/Private/Evaluation/MovieSceneEventTemplate.cpp
+++ b/Engine/Source/Runtime/MovieSceneTracks/Private/Evaluation/MovieSceneEventTemplate.cpp
@@ -9,7 +9,12 @@
 #include "Evaluation/MovieSceneEvaluation.h"
 #include "IMovieScenePlayer.h"
 #include "Algo/Accumulate.h"
+#include "Engine/LevelScriptActor.h"
+#include "Logging/TokenizedMessage.h"
+#include "Logging/MessageLog.h"
+#include "Misc/UObjectToken.h"
 
+#define LOCTEXT_NAMESPACE "MovieSceneEventTemplate"
 
 DECLARE_CYCLE_STAT(TEXT("Event Track Token Execute"), MovieSceneEval_EventTrack_TokenExecute, STATGROUP_MovieSceneEval);
 
@@ -238,7 +243,10 @@ struct FEventTriggerExecutionToken
 			// Event must have only a single object parameter, and the director instance must be an implementation of the function's class
 			if (!Function)
 			{
-				UE_LOG(LogMovieScene, Warning, TEXT("Failed to trigger the event '%s' because the function does not exist on director instance '%s'."), *EventName.ToString(), *DirectorInstance->GetName());
+				FMessageLog("PIE").Warning()
+					->AddToken(FTextToken::Create(FText::Format(LOCTEXT("LevelBP_MissingEvent_Error1", "Failed to trigger event '{0}' for"), FText::FromName(EventName))))
+					->AddToken(FUObjectToken::Create(Player.GetEvaluationTemplate().GetSequence(Operand.SequenceID)))
+					->AddToken(FTextToken::Create(LOCTEXT("LevelBP_MissingEvent_Error2", "because the function does not exist on the director instance.")));
 				continue;
 			}
 #if WITH_EDITOR
@@ -256,19 +264,43 @@ struct FEventTriggerExecutionToken
 			}
 			else if (Function->NumParms == 1 && Function->PropertyLink && (Function->PropertyLink->GetPropertyFlags() & CPF_ReferenceParm) == 0)
 			{
+				const int32 NumLevelScripts = Algo::Accumulate(EventContexts, 0, 
+					[](int32 Count, UObject* Obj)
+					{
+						return Obj && Obj->IsA<ALevelScriptActor>() ? Count + 1 : Count;
+					}
+				);
+
+				// Never pass through level script actors to event endpoints on non-interface pins.
+				if (NumLevelScripts > 0 && NumLevelScripts == EventContexts.Num() && !Function->PropertyLink->IsA<UInterfaceProperty>())
+				{
+					FMessageLog("PIE").Warning()
+						->AddToken(FTextToken::Create(LOCTEXT("LevelBP_ObjectPin_Error1", "Failed to trigger event")))
+						->AddToken(FUObjectToken::Create(Function))
+						->AddToken(FTextToken::Create(LOCTEXT("LevelBP_ObjectPin_Error2", "within")))
+						->AddToken(FUObjectToken::Create(Player.GetEvaluationTemplate().GetSequence(Operand.SequenceID)))
+						->AddToken(FTextToken::Create(LOCTEXT("LevelBP_ObjectPin_Error3", "because only Interface pins are supported for master tracks within Level Sequences. Please remove the pin, or change it to an interface that is implemented on the desired level blueprint.")));
+					continue;
+				}
+
 				for (UObject* EventContextObject : EventContexts)
 				{
-					TriggerEvent(DirectorInstance, Function, EventContextObject);
+					TriggerEvent(DirectorInstance, Function, EventContextObject, Player, Operand.SequenceID);
 				}
 			}
 			else
 			{
-				UE_LOG(LogMovieScene, Warning, TEXT("Failed to trigger event '%s' because its signature is not compatible. Function signatures must have either 0 or 1 (non-ref) parameters."), *EventName.ToString());
+				FMessageLog("PIE").Warning()
+					->AddToken(FTextToken::Create(LOCTEXT("LevelBP_InvalidEvent_Error1", "Failed to trigger event")))
+					->AddToken(FUObjectToken::Create(Function))
+					->AddToken(FTextToken::Create(LOCTEXT("LevelBP_InvalidEvent_Error2", "within")))
+					->AddToken(FUObjectToken::Create(Player.GetEvaluationTemplate().GetSequence(Operand.SequenceID)))
+					->AddToken(FTextToken::Create(LOCTEXT("LevelBP_InvalidEvent_Error3", "because its signature is not compatible. Function signatures must have either 0 or 1 (non-ref) parameters.")));
 			}
 		}
 	}
 
-	void TriggerEvent(UObject* DirectorInstance, UFunction* Function, UObject* ObjectParamValue)
+	void TriggerEvent(UObject* DirectorInstance, UFunction* Function, UObject* ObjectParamValue, IMovieScenePlayer& Player, FMovieSceneSequenceID SequenceID)
 	{
 		// We know by now that the parameter type is compatible because it wouldn't be added to the array if not
 		if (UObjectProperty* ObjectParameter = Cast<UObjectProperty>(Function->PropertyLink))
@@ -308,11 +340,12 @@ struct FEventTriggerExecutionToken
 			return;
 		}
 
-
-		UE_LOG(LogMovieScene, Warning, TEXT("Failed to trigger event '%s' because its signature is not compatible. Function expects a '%s' parameter, but only object and interface parameters are supported."),
-			*Function->GetName(),
-			*Function->PropertyLink->GetClass()->GetName()
-			);
+		FMessageLog("PIE").Warning()
+			->AddToken(FTextToken::Create(LOCTEXT("LevelBP_InvalidObjectEvent_Error1", "Failed to trigger event")))
+			->AddToken(FUObjectToken::Create(Function))
+			->AddToken(FTextToken::Create(LOCTEXT("LevelBP_InvalidObjectEvent_Error2", "within")))
+			->AddToken(FUObjectToken::Create(Player.GetEvaluationTemplate().GetSequence(SequenceID)))
+			->AddToken(FTextToken::Create(FText::Format(LOCTEXT("LevelBP_InvalidObjectEvent_Error3", "because its signature is not compatible. Function expects a '%s' parameter, but only object and interface parameters are supported."), FText::FromName(Function->PropertyLink->GetClass()->GetFName()))));
 	}
 
 #if !NO_LOGGING
@@ -502,4 +535,6 @@ void FMovieSceneEventRepeaterTemplate::EvaluateSwept(const FMovieSceneEvaluation
 		TArray<FName> Events = { EventToTrigger };
 		ExecutionTokens.Add(FEventTriggerExecutionToken(MoveTemp(Events), EventReceivers));
 	}
-}
\ No newline at end of file
+}
+
+#undef LOCTEXT_NAMESPACE
\ No newline at end of file
diff --git a/Engine/Source/Runtime/MovieSceneTracks/Private/Evaluation/MovieSceneSkeletalAnimationTemplate.cpp b/Engine/Source/Runtime/MovieSceneTracks/Private/Evaluation/MovieSceneSkeletalAnimationTemplate.cpp
--- a/Engine/Source/Runtime/MovieSceneTracks/Private/Evaluation/MovieSceneSkeletalAnimationTemplate.cpp
+++ b/Engine/Source/Runtime/MovieSceneTracks/Private/Evaluation/MovieSceneSkeletalAnimationTemplate.cpp
@@ -287,7 +287,7 @@ namespace MovieScene
 				FMontagePlayerPerSectionData* SectionData = MontageData.Find(Section);
 
 				int32 InstanceId = (SectionData) ? SectionData->MontageInstanceId : INDEX_NONE;
-				TWeakObjectPtr<UAnimMontage> Montage = FAnimMontageInstance::SetSequencerMontagePosition(SlotName, SkeletalMeshComponent, InstanceId, InAnimSequence, InPosition, Weight, bLooping);
+				TWeakObjectPtr<UAnimMontage> Montage = FAnimMontageInstance::SetSequencerMontagePosition(SlotName, SkeletalMeshComponent, InstanceId, InAnimSequence, InPosition, Weight, bLooping, bPlaying);
 
 				if (Montage.IsValid())
 				{
@@ -327,7 +327,7 @@ namespace MovieScene
 				FMontagePlayerPerSectionData* SectionData = MontageData.Find(Section);
 
 				int32 InstanceId = (SectionData)? SectionData->MontageInstanceId : INDEX_NONE;
-				TWeakObjectPtr<UAnimMontage> Montage = FAnimMontageInstance::PreviewSequencerMontagePosition(SlotName, SkeletalMeshComponent, InstanceId, InAnimSequence, InPosition, Weight, bLooping, bFireNotifies);
+				TWeakObjectPtr<UAnimMontage> Montage = FAnimMontageInstance::PreviewSequencerMontagePosition(SlotName, SkeletalMeshComponent, InstanceId, InAnimSequence, InPosition, Weight, bLooping, bFireNotifies, bPlaying);
 
 				if (Montage.IsValid())
 				{
diff --git a/Engine/Source/Runtime/MovieSceneTracks/Private/Sections/MovieScene3DConstraintSection.cpp b/Engine/Source/Runtime/MovieSceneTracks/Private/Sections/MovieScene3DConstraintSection.cpp
--- a/Engine/Source/Runtime/MovieSceneTracks/Private/Sections/MovieScene3DConstraintSection.cpp
+++ b/Engine/Source/Runtime/MovieSceneTracks/Private/Sections/MovieScene3DConstraintSection.cpp
@@ -17,13 +17,6 @@ void UMovieScene3DConstraintSection::SetConstraintId(const FGuid& InConstraintId
 	}
 }
 
-
-FGuid UMovieScene3DConstraintSection::GetConstraintId() const
-{
-	return FGuid();
-}
-
-
 void UMovieScene3DConstraintSection::OnBindingsUpdated(const TMap<FGuid, FGuid>& OldGuidToNewGuidMap)
 {
 	if (OldGuidToNewGuidMap.Contains(ConstraintBindingID.GetGuid()))
diff --git a/Engine/Source/Runtime/MovieSceneTracks/Private/Sections/MovieScene3DTransformSection.cpp b/Engine/Source/Runtime/MovieSceneTracks/Private/Sections/MovieScene3DTransformSection.cpp
--- a/Engine/Source/Runtime/MovieSceneTracks/Private/Sections/MovieScene3DTransformSection.cpp
+++ b/Engine/Source/Runtime/MovieSceneTracks/Private/Sections/MovieScene3DTransformSection.cpp
@@ -471,7 +471,7 @@ TSharedPtr<FStructOnScope> UMovieScene3DTransformSection::GetKeyStruct(TArrayVie
 		return KeyStruct;
 	}
 
-	if (AnyLocationKeys > 1)
+	if (AnyLocationKeys)
 	{
 		TSharedRef<FStructOnScope> KeyStruct = MakeShareable(new FStructOnScope(FMovieScene3DLocationKeyStruct::StaticStruct()));
 		auto Struct = (FMovieScene3DLocationKeyStruct*)KeyStruct->GetStructMemory();
diff --git a/Engine/Source/Runtime/MovieSceneTracks/Private/Sections/MovieSceneVectorSection.cpp b/Engine/Source/Runtime/MovieSceneTracks/Private/Sections/MovieSceneVectorSection.cpp
--- a/Engine/Source/Runtime/MovieSceneTracks/Private/Sections/MovieSceneVectorSection.cpp
+++ b/Engine/Source/Runtime/MovieSceneTracks/Private/Sections/MovieSceneVectorSection.cpp
@@ -48,7 +48,7 @@ struct FVectorSectionEditorData
 		}
 		else
 		{
-			return Bindings.GetCurrentValue<FVector>(InObject);
+			return Bindings.GetCurrentValue<FVector4>(InObject);
 		}
 	}
 
diff --git a/Engine/Source/Runtime/MovieSceneTracks/Private/Tracks/MovieSceneObjectPropertyTrack.cpp b/Engine/Source/Runtime/MovieSceneTracks/Private/Tracks/MovieSceneObjectPropertyTrack.cpp
--- a/Engine/Source/Runtime/MovieSceneTracks/Private/Tracks/MovieSceneObjectPropertyTrack.cpp
+++ b/Engine/Source/Runtime/MovieSceneTracks/Private/Tracks/MovieSceneObjectPropertyTrack.cpp
@@ -13,7 +13,7 @@ UMovieSceneObjectPropertyTrack::UMovieSceneObjectPropertyTrack(const FObjectInit
 
 UMovieSceneSection* UMovieSceneObjectPropertyTrack::CreateNewSection()
 {
-	UMovieSceneObjectPropertySection* Section = NewObject<UMovieSceneObjectPropertySection>(this);
+	UMovieSceneObjectPropertySection* Section = NewObject<UMovieSceneObjectPropertySection>(this, NAME_None, RF_Transactional);
 	Section->ObjectChannel.SetPropertyClass(PropertyClass);
 	return Section;
 }
diff --git a/Engine/Source/Runtime/MovieSceneTracks/Private/Tracks/MovieScenePrimitiveMaterialTrack.cpp b/Engine/Source/Runtime/MovieSceneTracks/Private/Tracks/MovieScenePrimitiveMaterialTrack.cpp
--- a/Engine/Source/Runtime/MovieSceneTracks/Private/Tracks/MovieScenePrimitiveMaterialTrack.cpp
+++ b/Engine/Source/Runtime/MovieSceneTracks/Private/Tracks/MovieScenePrimitiveMaterialTrack.cpp
@@ -16,7 +16,7 @@ UMovieScenePrimitiveMaterialTrack::UMovieScenePrimitiveMaterialTrack(const FObje
 
 UMovieSceneSection* UMovieScenePrimitiveMaterialTrack::CreateNewSection()
 {
-	return NewObject<UMovieScenePrimitiveMaterialSection>(this);
+	return NewObject<UMovieScenePrimitiveMaterialSection>(this, NAME_None, RF_Transactional);
 }
 
 FMovieSceneEvalTemplatePtr UMovieScenePrimitiveMaterialTrack::CreateTemplateForSection(const UMovieSceneSection& InSection) const
diff --git a/Engine/Source/Runtime/MovieSceneTracks/Public/Sections/MovieScene3DConstraintSection.h b/Engine/Source/Runtime/MovieSceneTracks/Public/Sections/MovieScene3DConstraintSection.h
--- a/Engine/Source/Runtime/MovieSceneTracks/Public/Sections/MovieScene3DConstraintSection.h
+++ b/Engine/Source/Runtime/MovieSceneTracks/Public/Sections/MovieScene3DConstraintSection.h
@@ -25,9 +25,6 @@ public:
 	/** Sets the constraint id for this section */
 	virtual void SetConstraintId(const FGuid& InId);
 
-	UE_DEPRECATED(4.18, "Constraint guid no longer supported, Use GetConstraintBindingID.")
-	virtual FGuid GetConstraintId() const;
-
 	/** Gets the constraint binding for this Constraint section */
 	const FMovieSceneObjectBindingID& GetConstraintBindingID() const
 	{
diff --git a/Engine/Source/Runtime/MovieSceneTracks/Public/Sections/MovieSceneCameraCutSection.h b/Engine/Source/Runtime/MovieSceneTracks/Public/Sections/MovieSceneCameraCutSection.h
--- a/Engine/Source/Runtime/MovieSceneTracks/Public/Sections/MovieSceneCameraCutSection.h
+++ b/Engine/Source/Runtime/MovieSceneTracks/Public/Sections/MovieSceneCameraCutSection.h
@@ -31,12 +31,6 @@ public:
 				EMovieSceneCompletionMode::ProjectDefault);
 	}
 
-	UE_DEPRECATED(4.18, "Camera guid no longer supported, Use GetCameraBindingID.")
-	FGuid GetCameraGuid() const
-	{
-		return FGuid();
-	}
-
 	/** Sets the camera binding for this CameraCut section. Evaluates from the sequence binding ID */
 	void SetCameraGuid(const FGuid& InGuid)
 	{
