PR #5407: Added missing icons to editor (Contributed by [REDACTED_PERSON])

diff --git a/Engine/Plugins/MovieScene/ActorSequence/Source/ActorSequenceEditor/Private/ActorSequenceEditorTabSummoner.cpp b/Engine/Plugins/MovieScene/ActorSequence/Source/ActorSequenceEditor/Private/ActorSequenceEditorTabSummoner.cpp
--- a/Engine/Plugins/MovieScene/ActorSequence/Source/ActorSequenceEditor/Private/ActorSequenceEditorTabSummoner.cpp
+++ b/Engine/Plugins/MovieScene/ActorSequence/Source/ActorSequenceEditor/Private/ActorSequenceEditorTabSummoner.cpp
@@ -15,7 +15,7 @@
 #include "ScopedTransaction.h"
 #include "Framework/MultiBox/MultiBoxBuilder.h"
 #include "Framework/Application/SlateApplication.h"
-
+#include "ActorSequenceEditorStyle.h"
 
 #define LOCTEXT_NAMESPACE "ActorSequenceEditorSummoner"
 
@@ -521,6 +521,7 @@ FActorSequenceEditorSummoner::FActorSequenceEditorSummoner(TSharedPtr<FBlueprint
 	bIsSingleton = true;
 
 	TabLabel = LOCTEXT("SequencerTabName", "Sequencer");
+	TabIcon = FSlateIcon(FActorSequenceEditorStyle::Get().GetStyleSetName(), "ClassIcon.ActorSequence");
 }
 
 TSharedRef<SWidget> FActorSequenceEditorSummoner::CreateTabBody(const FWorkflowTabSpawnInfo& Info) const
diff --git a/Engine/Source/Editor/BehaviorTreeEditor/Private/BehaviorTreeEditorTabFactories.cpp b/Engine/Source/Editor/BehaviorTreeEditor/Private/BehaviorTreeEditorTabFactories.cpp
--- a/Engine/Source/Editor/BehaviorTreeEditor/Private/BehaviorTreeEditorTabFactories.cpp
+++ b/Engine/Source/Editor/BehaviorTreeEditor/Private/BehaviorTreeEditorTabFactories.cpp
@@ -108,7 +108,7 @@ FBehaviorTreeSearchSummoner::FBehaviorTreeSearchSummoner(TSharedPtr<class FBehav
 	, BehaviorTreeEditorPtr(InBehaviorTreeEditorPtr)
 {
 	TabLabel = LOCTEXT("BehaviorTreeSearchLabel", "Search");
-	TabIcon = FSlateIcon(FEditorStyle::GetStyleSetName(), "Kismet.Tabs.Components");
+	TabIcon = FSlateIcon(FEditorStyle::GetStyleSetName(), "Kismet.Tabs.FindResults");
 
 	bIsSingleton = true;
 
diff --git a/Engine/Source/Editor/MovieSceneTools/Private/TrackEditors/TransformTrackEditor.cpp b/Engine/Source/Editor/MovieSceneTools/Private/TrackEditors/TransformTrackEditor.cpp
--- a/Engine/Source/Editor/MovieSceneTools/Private/TrackEditors/TransformTrackEditor.cpp
+++ b/Engine/Source/Editor/MovieSceneTools/Private/TrackEditors/TransformTrackEditor.cpp
@@ -226,13 +226,11 @@ void F3DTransformTrackEditor::BuildTrackContextMenu( FMenuBuilder& MenuBuilder,
 			FText(), true, false);
 	};
 
-	const bool bInOpenSubMenuOnClick = true;
-
 	MenuBuilder.AddSubMenu(
 		NSLOCTEXT( "Sequencer", "ImportTransforms", "Import From Animation Root" ),
 		NSLOCTEXT( "Sequencer", "ImportTransformsTooltip", "Import transform keys from an animation sequence's root motion." ),
-		FNewMenuDelegate::CreateLambda(AnimSubMenuDelegate, GetSequencer().ToSharedRef(), TransformTrack), 
-		bInOpenSubMenuOnClick);
+		FNewMenuDelegate::CreateLambda(AnimSubMenuDelegate, GetSequencer().ToSharedRef(), TransformTrack)
+	);
 
 	MenuBuilder.AddMenuSeparator();
 	FKeyframeTrackEditor::BuildTrackContextMenu(MenuBuilder, Track);
