Fix non-unity compile error.

#rnx

diff --git a/Engine/Plugins/MovieScene/ActorSequence/Source/ActorSequenceEditor/Private/ActorSequenceEditorTabSummoner.cpp b/Engine/Plugins/MovieScene/ActorSequence/Source/ActorSequenceEditor/Private/ActorSequenceEditorTabSummoner.cpp
--- a/Engine/Plugins/MovieScene/ActorSequence/Source/ActorSequenceEditor/Private/ActorSequenceEditorTabSummoner.cpp
+++ b/Engine/Plugins/MovieScene/ActorSequence/Source/ActorSequenceEditor/Private/ActorSequenceEditorTabSummoner.cpp
@@ -13,6 +13,7 @@
 #include "Widgets/Images/SImage.h"
 #include "Editor.h"
 #include "ScopedTransaction.h"
+#include "Widgets/Docking/SDockTab.h"
 #include "Framework/MultiBox/MultiBoxBuilder.h"
 #include "Framework/Application/SlateApplication.h"
 #include "ActorSequenceEditorStyle.h"
