Actor Sequence: Find curve tab spawner before registering

diff --git a/Engine/Plugins/MovieScene/ActorSequence/Source/ActorSequenceEditor/Private/ActorSequenceEditorTabSummoner.cpp b/Engine/Plugins/MovieScene/ActorSequence/Source/ActorSequenceEditor/Private/ActorSequenceEditorTabSummoner.cpp
--- a/Engine/Plugins/MovieScene/ActorSequence/Source/ActorSequenceEditor/Private/ActorSequenceEditorTabSummoner.cpp
+++ b/Engine/Plugins/MovieScene/ActorSequence/Source/ActorSequenceEditor/Private/ActorSequenceEditorTabSummoner.cpp
@@ -224,10 +224,14 @@ public:
 		WeakBlueprintEditor = InBlueprintEditor;
 
 		{
-			// Register an empty tab to spawn the Curve Editor in so that layouts restore properly.
-			WeakBlueprintEditor.Pin()->GetTabManager()->RegisterTabSpawner("SequencerGraphEditor", 
-				FOnSpawnTab::CreateSP(this, &SActorSequenceEditorWidgetImpl::SpawnCurveEditorTab))
-				.SetMenuType(ETabSpawnerMenuType::Type::Hidden);
+			const FName CurveEditorTabName = FName(TEXT("SequencerGraphEditor"));
+			if (!WeakBlueprintEditor.Pin()->GetTabManager()->FindExistingLiveTab(CurveEditorTabName))
+			{
+				// Register an empty tab to spawn the Curve Editor in so that layouts restore properly.
+				WeakBlueprintEditor.Pin()->GetTabManager()->RegisterTabSpawner(CurveEditorTabName,
+					FOnSpawnTab::CreateSP(this, &SActorSequenceEditorWidgetImpl::SpawnCurveEditorTab))
+					.SetMenuType(ETabSpawnerMenuType::Type::Hidden);
+			}
 		}
 
 		ChildSlot
