Actor Sequencer: Tab spawner fixes for graph editor

diff --git a/Engine/Plugins/MovieScene/ActorSequence/Source/ActorSequenceEditor/Private/ActorSequenceEditorTabSummoner.cpp b/Engine/Plugins/MovieScene/ActorSequence/Source/ActorSequenceEditor/Private/ActorSequenceEditorTabSummoner.cpp
--- a/Engine/Plugins/MovieScene/ActorSequence/Source/ActorSequenceEditor/Private/ActorSequenceEditorTabSummoner.cpp
+++ b/Engine/Plugins/MovieScene/ActorSequence/Source/ActorSequenceEditor/Private/ActorSequenceEditorTabSummoner.cpp
@@ -201,6 +201,16 @@ public:
 			Sequencer = nullptr;
 		}
 
+		if (WeakBlueprintEditor.IsValid())
+		{
+			const FName CurveEditorTabName = FName(TEXT("SequencerGraphEditor"));
+			TSharedPtr<SDockTab> ExistingTab = WeakBlueprintEditor.Pin()->GetToolkitHost()->GetTabManager()->FindExistingLiveTab(CurveEditorTabName);
+			if (ExistingTab)
+			{
+				ExistingTab->RequestCloseTab();
+			}
+		}
+
 		GEditor->OnBlueprintPreCompile().Remove(OnBlueprintPreCompileHandle);
 		FCoreUObjectDelegates::OnObjectSaved.Remove(OnObjectSavedHandle);
 	}
@@ -225,7 +235,7 @@ public:
 
 		{
 			const FName CurveEditorTabName = FName(TEXT("SequencerGraphEditor"));
-			if (!WeakBlueprintEditor.Pin()->GetTabManager()->FindExistingLiveTab(CurveEditorTabName))
+			if (!WeakBlueprintEditor.Pin()->GetTabManager()->HasTabSpawner(CurveEditorTabName))
 			{
 				// Register an empty tab to spawn the Curve Editor in so that layouts restore properly.
 				WeakBlueprintEditor.Pin()->GetTabManager()->RegisterTabSpawner(CurveEditorTabName,
