Actor Sequence: Check for valid blueprint editor in actor sequence tab

#rnx

diff --git a/Engine/Plugins/MovieScene/ActorSequence/Source/ActorSequenceEditor/Private/ActorSequenceEditorTabSummoner.cpp b/Engine/Plugins/MovieScene/ActorSequence/Source/ActorSequenceEditor/Private/ActorSequenceEditorTabSummoner.cpp
--- a/Engine/Plugins/MovieScene/ActorSequence/Source/ActorSequenceEditor/Private/ActorSequenceEditorTabSummoner.cpp
+++ b/Engine/Plugins/MovieScene/ActorSequence/Source/ActorSequenceEditor/Private/ActorSequenceEditorTabSummoner.cpp
@@ -235,7 +235,7 @@ public:
 
 		{
 			const FName CurveEditorTabName = FName(TEXT("SequencerGraphEditor"));
-			if (!WeakBlueprintEditor.Pin()->GetTabManager()->HasTabSpawner(CurveEditorTabName))
+			if (WeakBlueprintEditor.IsValid() && !WeakBlueprintEditor.Pin()->GetTabManager()->HasTabSpawner(CurveEditorTabName))
 			{
 				// Register an empty tab to spawn the Curve Editor in so that layouts restore properly.
 				WeakBlueprintEditor.Pin()->GetTabManager()->RegisterTabSpawner(CurveEditorTabName,
