Fix deprecation warnings

diff --git a/Engine/Plugins/MovieScene/ActorSequence/Source/ActorSequenceEditor/Private/ActorSequenceComponentCustomization.cpp b/Engine/Plugins/MovieScene/ActorSequence/Source/ActorSequenceEditor/Private/ActorSequenceComponentCustomization.cpp
--- a/Engine/Plugins/MovieScene/ActorSequence/Source/ActorSequenceEditor/Private/ActorSequenceComponentCustomization.cpp
+++ b/Engine/Plugins/MovieScene/ActorSequence/Source/ActorSequenceEditor/Private/ActorSequenceComponentCustomization.cpp
@@ -100,7 +100,7 @@ void FActorSequenceComponentCustomization::CustomizeDetails(IDetailLayoutBuilder
 
 	bool bIsExternalTabAlreadyOpened = false;
 
-	if (HostTabManager.IsValid() && HostTabManager->CanSpawnTab(SequenceTabId))
+	if (HostTabManager.IsValid() && HostTabManager->HasTabSpawner(SequenceTabId))
 	{
 		WeakTabManager = HostTabManager;
 
@@ -151,7 +151,7 @@ void FActorSequenceComponentCustomization::CustomizeDetails(IDetailLayoutBuilder
 FReply FActorSequenceComponentCustomization::InvokeSequencer()
 {
 	TSharedPtr<FTabManager> TabManager = WeakTabManager.Pin();
-	if (TabManager.IsValid() && TabManager->CanSpawnTab(SequenceTabId))
+	if (TabManager.IsValid() && TabManager->HasTabSpawner(SequenceTabId))
 	{
 		TSharedRef<SDockTab> Tab = TabManager->InvokeTab(SequenceTabId);
 
diff --git a/Engine/Source/Editor/Blutility/Private/AssetTypeActions_EditorUtilityWidgetBlueprint.cpp b/Engine/Source/Editor/Blutility/Private/AssetTypeActions_EditorUtilityWidgetBlueprint.cpp
--- a/Engine/Source/Editor/Blutility/Private/AssetTypeActions_EditorUtilityWidgetBlueprint.cpp
+++ b/Engine/Source/Editor/Blutility/Private/AssetTypeActions_EditorUtilityWidgetBlueprint.cpp
@@ -130,7 +130,7 @@ void FAssetTypeActions_EditorUtilityWidgetBlueprint::ExecuteRun(FWeakBlueprintPo
 					FText DisplayName = FText::FromString(Blueprint->GetName());
 					FLevelEditorModule& LevelEditorModule = FModuleManager::GetModuleChecked<FLevelEditorModule>(TEXT("LevelEditor"));
 					TSharedPtr<FTabManager> LevelEditorTabManager = LevelEditorModule.GetLevelEditorTabManager();
-					if (!LevelEditorTabManager->CanSpawnTab(RegistrationName))
+					if (!LevelEditorTabManager->HasTabSpawner(RegistrationName))
 					{
 						IBlutilityModule* BlutilityModule = FModuleManager::GetModulePtr<IBlutilityModule>("Blutility");
 						UEditorUtilityWidgetBlueprint* WidgetBlueprint = Cast<UEditorUtilityWidgetBlueprint>(Blueprint);
diff --git a/Engine/Source/Editor/Blutility/Private/BlutilityModule.cpp b/Engine/Source/Editor/Blutility/Private/BlutilityModule.cpp
--- a/Engine/Source/Editor/Blutility/Private/BlutilityModule.cpp
+++ b/Engine/Source/Editor/Blutility/Private/BlutilityModule.cpp
@@ -117,7 +117,7 @@ public:
 						FName RegistrationName = FName(*(Blueprint->GetPathName() + LOCTEXT("ActiveTabSuffix", "_ActiveTab").ToString()));
 						Blueprint->SetRegistrationName(RegistrationName);
 						FText DisplayName = FText::FromString(Blueprint->GetName());
-						if (LevelEditorTabManager && !LevelEditorTabManager->CanSpawnTab(RegistrationName))
+						if (LevelEditorTabManager && !LevelEditorTabManager->HasTabSpawner(RegistrationName))
 						{
 							LevelEditorTabManager->RegisterTabSpawner(RegistrationName, FOnSpawnTab::CreateUObject(Blueprint, &UEditorUtilityWidgetBlueprint::SpawnEditorUITab))
 								.SetDisplayName(DisplayName)
diff --git a/Engine/Source/Editor/ClothPainter/Private/ClothPaintingModule.cpp b/Engine/Source/Editor/ClothPainter/Private/ClothPaintingModule.cpp
--- a/Engine/Source/Editor/ClothPainter/Private/ClothPaintingModule.cpp
+++ b/Engine/Source/Editor/ClothPainter/Private/ClothPaintingModule.cpp
@@ -175,7 +175,7 @@ TSharedPtr<SClothPaintTab> FClothPaintingModule::GetActiveClothTab(TWeakPtr<ISke
 
 	// If we can't summon this tab we will have spawned a placeholder which we can not cast as done below,
 	// there is no valid active clothing tab so return nullptr
-	if(TabManager->CanSpawnTab(FClothPaintTabSummoner::TabName))
+	if(TabManager->HasTabSpawner(FClothPaintTabSummoner::TabName))
 	{
 		TSharedPtr<SDockTab> Tab = TabManager->FindExistingLiveTab(FTabId(FClothPaintTabSummoner::TabName));
 
diff --git a/Engine/Source/Editor/Kismet/Private/BlueprintEditor.cpp b/Engine/Source/Editor/Kismet/Private/BlueprintEditor.cpp
--- a/Engine/Source/Editor/Kismet/Private/BlueprintEditor.cpp
+++ b/Engine/Source/Editor/Kismet/Private/BlueprintEditor.cpp
@@ -8508,7 +8508,7 @@ bool FBlueprintEditor::IsFocusedGraphEditable() const
 
 void FBlueprintEditor::TryInvokingDetailsTab(bool bFlash)
 {
-	if ( TabManager->CanSpawnTab(FBlueprintEditorTabs::DetailsID) )
+	if ( TabManager->HasTabSpawner(FBlueprintEditorTabs::DetailsID) )
 	{
 		TSharedPtr<SDockTab> BlueprintTab = FGlobalTabmanager::Get()->GetMajorTabForTabManager(TabManager.ToSharedRef());
 
diff --git a/Engine/Source/Editor/Kismet/Private/FindInBlueprintManager.cpp b/Engine/Source/Editor/Kismet/Private/FindInBlueprintManager.cpp
--- a/Engine/Source/Editor/Kismet/Private/FindInBlueprintManager.cpp
+++ b/Engine/Source/Editor/Kismet/Private/FindInBlueprintManager.cpp
@@ -2623,7 +2623,7 @@ void FFindInBlueprintSearchManager::EnableGlobalFindResults(bool bEnable)
 		for (int32 TabIdx = 0; TabIdx < ARRAY_COUNT(GlobalFindResultsTabIDs); TabIdx++)
 		{
 			const FName TabID = GlobalFindResultsTabIDs[TabIdx];
-			if (GlobalTabManager->CanSpawnTab(TabID))
+			if (GlobalTabManager->HasTabSpawner(TabID))
 			{
 				GlobalTabManager->UnregisterNomadTabSpawner(TabID);
 			}
@@ -2644,7 +2644,7 @@ void FFindInBlueprintSearchManager::CloseOrphanedGlobalFindResultsTabs(TSharedPt
 		for (int32 TabIdx = 0; TabIdx < ARRAY_COUNT(GlobalFindResultsTabIDs); TabIdx++)
 		{
 			const FName TabID = GlobalFindResultsTabIDs[TabIdx];
-			if (!FGlobalTabmanager::Get()->CanSpawnTab(TabID))
+			if (!FGlobalTabmanager::Get()->HasTabSpawner(TabID))
 			{
 				TSharedPtr<SDockTab> OrphanedTab = TabManager->FindExistingLiveTab(FTabId(TabID));
 				if (OrphanedTab.IsValid())
diff --git a/Engine/Source/Runtime/Slate/Public/Framework/Docking/TabManager.h b/Engine/Source/Runtime/Slate/Public/Framework/Docking/TabManager.h
--- a/Engine/Source/Runtime/Slate/Public/Framework/Docking/TabManager.h
+++ b/Engine/Source/Runtime/Slate/Public/Framework/Docking/TabManager.h
@@ -761,6 +761,7 @@ class SLATE_API FTabManager : public TSharedFromThis<FTabManager>
 		UE_DEPRECATED(4.23, "CanSpawnTab has been replaced by HasTabSpawner")
 		bool CanSpawnTab(FName TabId) const;
 
+		/** @return true if the tab has a factory registered for it that allows it to be spawned. */
 		bool HasTabSpawner(FName TabId) const;
 
 		/** Returns the owner tab (if it exists) */
