Actor Sequence: Check blueprint editor when closing curve editor tab for actor sequence. This fixes a crash when invoking the merge tool.

diff --git a/Engine/Plugins/MovieScene/ActorSequence/Source/ActorSequenceEditor/Private/ActorSequenceEditorTabSummoner.cpp b/Engine/Plugins/MovieScene/ActorSequence/Source/ActorSequenceEditor/Private/ActorSequenceEditorTabSummoner.cpp
--- a/Engine/Plugins/MovieScene/ActorSequence/Source/ActorSequenceEditor/Private/ActorSequenceEditorTabSummoner.cpp
+++ b/Engine/Plugins/MovieScene/ActorSequence/Source/ActorSequenceEditor/Private/ActorSequenceEditorTabSummoner.cpp
@@ -201,7 +201,7 @@ public:
 			Sequencer = nullptr;
 		}
 
-		if (WeakBlueprintEditor.IsValid())
+		if (WeakBlueprintEditor.IsValid() && WeakBlueprintEditor.Pin()->IsHosted())
 		{
 			const FName CurveEditorTabName = FName(TEXT("SequencerGraphEditor"));
 			TSharedPtr<SDockTab> ExistingTab = WeakBlueprintEditor.Pin()->GetToolkitHost()->GetTabManager()->FindExistingLiveTab(CurveEditorTabName);
