Actor Sequence: Close on MapChanged. Otherwise FSequencer leaks memory. This follows the pattern in other cases ie. TemplateSequenceEditorToolkit.cpp

diff --git a/Engine/Plugins/MovieScene/ActorSequence/Source/ActorSequenceEditor/Private/ActorSequenceEditorTabSummoner.cpp b/Engine/Plugins/MovieScene/ActorSequence/Source/ActorSequenceEditor/Private/ActorSequenceEditorTabSummoner.cpp
--- a/Engine/Plugins/MovieScene/ActorSequence/Source/ActorSequenceEditor/Private/ActorSequenceEditorTabSummoner.cpp
+++ b/Engine/Plugins/MovieScene/ActorSequence/Source/ActorSequenceEditor/Private/ActorSequenceEditorTabSummoner.cpp
@@ -13,6 +13,7 @@
 #include "EditorUndoClient.h"
 #include "Widgets/Images/SImage.h"
 #include "Editor.h"
+#include "LevelEditor.h"
 #include "ScopedTransaction.h"
 #include "Widgets/Docking/SDockTab.h"
 #include "Framework/MultiBox/MultiBoxBuilder.h"
@@ -414,6 +415,17 @@ public:
 		Options.bForceRefreshDetails = false;
 
 		FLevelEditorSequencerIntegration::Get().AddSequencer(Sequencer.ToSharedRef(), Options);
+	
+		FLevelEditorModule& LevelEditorModule = FModuleManager::LoadModuleChecked<FLevelEditorModule>("LevelEditor");
+		LevelEditorModule.OnMapChanged().AddRaw(this, &SActorSequenceEditorWidgetImpl::HandleMapChanged);
+	}
+
+	void HandleMapChanged(UWorld* NewWorld, EMapChangeType MapChangeType)
+	{
+		if ((MapChangeType == EMapChangeType::LoadMap || MapChangeType == EMapChangeType::NewMap || MapChangeType == EMapChangeType::TearDownWorld))
+		{
+			Close();
+		}
 	}
 
 	void Refresh()
