PR #9357: Add Ease of Use functionality got ActorSequenceComponent (Contributed by [REDACTED_PERSON])

diff --git a/Engine/Plugins/MovieScene/ActorSequence/Source/ActorSequence/Private/ActorSequenceComponent.cpp b/Engine/Plugins/MovieScene/ActorSequence/Source/ActorSequence/Private/ActorSequenceComponent.cpp
--- a/Engine/Plugins/MovieScene/ActorSequence/Source/ActorSequence/Private/ActorSequenceComponent.cpp
+++ b/Engine/Plugins/MovieScene/ActorSequence/Source/ActorSequence/Private/ActorSequenceComponent.cpp
@@ -16,6 +16,30 @@ UActorSequenceComponent::UActorSequenceComponent(const FObjectInitializer& Objec
 	}
 }
 
+void UActorSequenceComponent::PlaySequence()
+{
+	if (SequencePlayer != nullptr)
+	{
+		SequencePlayer->Play();
+	}
+}
+
+void UActorSequenceComponent::PauseSequence()
+{
+	if (SequencePlayer != nullptr)
+	{
+		SequencePlayer->Pause();
+	}
+}
+
+void UActorSequenceComponent::StopSequence()
+{
+	if (SequencePlayer != nullptr)
+	{
+		SequencePlayer->Stop();
+	}
+}
+
 void UActorSequenceComponent::PostInitProperties()
 {
 	Super::PostInitProperties();
diff --git a/Engine/Plugins/MovieScene/ActorSequence/Source/ActorSequence/Public/ActorSequenceComponent.h b/Engine/Plugins/MovieScene/ActorSequence/Source/ActorSequence/Public/ActorSequenceComponent.h
--- a/Engine/Plugins/MovieScene/ActorSequence/Source/ActorSequence/Public/ActorSequenceComponent.h
+++ b/Engine/Plugins/MovieScene/ActorSequence/Source/ActorSequence/Public/ActorSequenceComponent.h
@@ -32,6 +32,18 @@ public:
 	{
 		return SequencePlayer;
 	}
+
+	/** Calls the Play function on the SequencePlayer if its valid. */
+	UFUNCTION(BlueprintCallable, Category = "Sequencer|Player")
+	void PlaySequence();
+
+	/** Calls the Pause function on the SequencePlayer if its valid. */
+	UFUNCTION(BlueprintCallable, Category = "Sequencer|Player")
+	void PauseSequence();
+
+	/** Calls the Stop function on the SequencePlayer if its valid. */
+	UFUNCTION(BlueprintCallable, Category = "Sequencer|Player")
+	void StopSequence();
 	
 	// UActorComponent interface
 	virtual void PostInitProperties() override;
