Actor Sequence: Clear reference to sequencer widget when closing. This fixes a crash when opening a new map and there is a BP with an active actor sequence.

diff --git a/Engine/Plugins/MovieScene/ActorSequence/Source/ActorSequenceEditor/Private/ActorSequenceEditorTabSummoner.cpp b/Engine/Plugins/MovieScene/ActorSequence/Source/ActorSequenceEditor/Private/ActorSequenceEditorTabSummoner.cpp
--- a/Engine/Plugins/MovieScene/ActorSequence/Source/ActorSequenceEditor/Private/ActorSequenceEditorTabSummoner.cpp
+++ b/Engine/Plugins/MovieScene/ActorSequence/Source/ActorSequenceEditor/Private/ActorSequenceEditorTabSummoner.cpp
@@ -215,6 +215,8 @@ public:
 			Sequencer = nullptr;
 		}
 
+		Content->SetContent(SNew(STextBlock).Text(LOCTEXT("NothingSelected", "Select a sequence")));
+
 		if (WeakBlueprintEditor.IsValid() && WeakBlueprintEditor.Pin()->IsHosted())
 		{
 			const FName CurveEditorTabName = FName(TEXT("SequencerGraphEditor"));
