Actor Sequence: Remove OnMapChanged on close to fix crash in Actor Sequence widget

diff --git a/Engine/Plugins/MovieScene/ActorSequence/Source/ActorSequenceEditor/Private/ActorSequenceEditorTabSummoner.cpp b/Engine/Plugins/MovieScene/ActorSequence/Source/ActorSequenceEditor/Private/ActorSequenceEditorTabSummoner.cpp
--- a/Engine/Plugins/MovieScene/ActorSequence/Source/ActorSequenceEditor/Private/ActorSequenceEditorTabSummoner.cpp
+++ b/Engine/Plugins/MovieScene/ActorSequence/Source/ActorSequenceEditor/Private/ActorSequenceEditorTabSummoner.cpp
@@ -215,7 +215,15 @@ public:
 			Sequencer = nullptr;
 		}
 
-		Content->SetContent(SNew(STextBlock).Text(LOCTEXT("NothingSelected", "Select a sequence")));
+		if (FLevelEditorModule* LevelEditorModule = FModuleManager::GetModulePtr<FLevelEditorModule>("LevelEditor"))
+		{
+			LevelEditorModule->OnMapChanged().RemoveAll(this);
+		}
+
+		if (Content)
+		{
+			Content->SetContent(SNew(STextBlock).Text(LOCTEXT("NothingSelected", "Select a sequence")));
+		}
 
 		if (WeakBlueprintEditor.IsValid() && WeakBlueprintEditor.Pin()->IsHosted())
 		{
