Sequencer: Time Warp track is supported by [REDACTED_PERSON], Widget Animation, Day Sequence, and Actor Sequences

#rnx

diff --git a/Engine/Plugins/Experimental/DaySequence/Source/DaySequence/Private/DaySequence.cpp b/Engine/Plugins/Experimental/DaySequence/Source/DaySequence/Private/DaySequence.cpp
--- a/Engine/Plugins/Experimental/DaySequence/Source/DaySequence/Private/DaySequence.cpp
+++ b/Engine/Plugins/Experimental/DaySequence/Source/DaySequence/Private/DaySequence.cpp
@@ -31,6 +31,7 @@
 #include "Tracks/MovieSceneSlomoTrack.h"
 #include "Tracks/MovieSceneSpawnTrack.h"
 #include "Tracks/MovieSceneSubTrack.h"
+#include "Tracks/MovieSceneTimeWarpTrack.h"
 #include "UObject/AssetRegistryTagsContext.h"
 
 #include UE_INLINE_GENERATED_CPP_BY_NAME(DaySequence)
@@ -129,6 +130,7 @@ ETrackSupport UDaySequence::IsTrackSupportedImpl(TSubclassOf<class UMovieSceneTr
 		InTrackClass == UMovieSceneMaterialParameterCollectionTrack::StaticClass() ||
 		InTrackClass == UMovieSceneSlomoTrack::StaticClass() ||
 		InTrackClass == UMovieSceneSpawnTrack::StaticClass() ||
+		InTrackClass == UMovieSceneTimeWarpTrack::StaticClass() ||
 		InTrackClass == UMovieSceneCVarTrack::StaticClass() ||
 		InTrackClass == UDaySequenceTrack::StaticClass())
 	{
diff --git a/Engine/Plugins/MovieScene/ActorSequence/Source/ActorSequence/Private/ActorSequence.cpp b/Engine/Plugins/MovieScene/ActorSequence/Source/ActorSequence/Private/ActorSequence.cpp
--- a/Engine/Plugins/MovieScene/ActorSequence/Source/ActorSequence/Private/ActorSequence.cpp
+++ b/Engine/Plugins/MovieScene/ActorSequence/Source/ActorSequence/Private/ActorSequence.cpp
@@ -13,6 +13,7 @@
 #include "Tracks/MovieSceneEventTrack.h"
 #include "Tracks/MovieSceneMaterialParameterCollectionTrack.h"
 #include "Tracks/MovieSceneSkeletalAnimationTrack.h"
+#include "Tracks/MovieSceneTimeWarpTrack.h"
 
 IMPLEMENT_MODULE(FDefaultModuleImpl, ActorSequence);
 
@@ -221,7 +222,8 @@ ETrackSupport UActorSequence::IsTrackSupportedImpl(TSubclassOf<class UMovieScene
 	if (InTrackClass == UMovieSceneAudioTrack::StaticClass() ||
 		InTrackClass == UMovieSceneEventTrack::StaticClass() ||
 		InTrackClass == UMovieSceneMaterialParameterCollectionTrack::StaticClass() ||
-		InTrackClass == UMovieSceneSkeletalAnimationTrack::StaticClass())
+		InTrackClass == UMovieSceneSkeletalAnimationTrack::StaticClass() ||
+		InTrackClass == UMovieSceneTimeWarpTrack::StaticClass())
 	{
 		return ETrackSupport::Supported;
 	}
diff --git a/Engine/Source/Editor/MovieSceneTools/Private/TrackEditors/TimeWarpTrackEditor.cpp b/Engine/Source/Editor/MovieSceneTools/Private/TrackEditors/TimeWarpTrackEditor.cpp
--- a/Engine/Source/Editor/MovieSceneTools/Private/TrackEditors/TimeWarpTrackEditor.cpp
+++ b/Engine/Source/Editor/MovieSceneTools/Private/TrackEditors/TimeWarpTrackEditor.cpp
@@ -368,6 +368,12 @@ TSharedRef<ISequencerSection> FTimeWarpTrackEditor::MakeSectionInterface(UMovieS
 	return MakeShared<UE::Sequencer::FTimeWarpSection>(SectionObject, GetSequencer());
 }
 
+bool FTimeWarpTrackEditor::SupportsSequence(UMovieSceneSequence* InSequence) const
+{
+	ETrackSupport TrackSupported = InSequence ? InSequence->IsTrackSupported(UMovieSceneTimeWarpTrack::StaticClass()) : ETrackSupport::NotSupported;
+	return TrackSupported == ETrackSupport::Supported;
+}
+
 void FTimeWarpTrackEditor::HandleAddTimeWarpTrack(TSubclassOf<UMovieSceneTimeWarpGetter> ClassType)
 {
 	UMovieScene* FocusedMovieScene = GetFocusedMovieScene();
diff --git a/Engine/Source/Editor/MovieSceneTools/Public/TrackEditors/TimeWarpTrackEditor.h b/Engine/Source/Editor/MovieSceneTools/Public/TrackEditors/TimeWarpTrackEditor.h
--- a/Engine/Source/Editor/MovieSceneTools/Public/TrackEditors/TimeWarpTrackEditor.h
+++ b/Engine/Source/Editor/MovieSceneTools/Public/TrackEditors/TimeWarpTrackEditor.h
@@ -74,6 +74,7 @@ private:
 	virtual void BuildAddTrackMenu(FMenuBuilder& MenuBuilder) override;
 	virtual void ProcessKeyOperation(FFrameNumber InKeyTime, const UE::Sequencer::FKeyOperation& Operation, ISequencer& InSequencer) override;
 	virtual TSharedRef<ISequencerSection> MakeSectionInterface(UMovieSceneSection& SectionObject, UMovieSceneTrack& Track, FGuid ObjectBinding) override;
+	virtual bool SupportsSequence(UMovieSceneSequence* InSequence) const override;
 
 	void HandleAddTimeWarpTrack(TSubclassOf<UMovieSceneTimeWarpGetter> ClassType);
 };
diff --git a/Engine/Source/Runtime/LevelSequence/Private/LevelSequence.cpp b/Engine/Source/Runtime/LevelSequence/Private/LevelSequence.cpp
--- a/Engine/Source/Runtime/LevelSequence/Private/LevelSequence.cpp
+++ b/Engine/Source/Runtime/LevelSequence/Private/LevelSequence.cpp
@@ -39,6 +39,7 @@
 #include "Tracks/MovieSceneSlomoTrack.h"
 #include "Tracks/MovieSceneSpawnTrack.h"
 #include "Tracks/MovieSceneSubTrack.h"
+#include "Tracks/MovieSceneTimeWarpTrack.h"
 #include "Tracks/MovieSceneCVarTrack.h"
 #include "Tracks/MovieSceneBindingLifetimeTrack.h"
 #include "Modules/ModuleManager.h"
@@ -154,7 +155,8 @@ ETrackSupport ULevelSequence::IsTrackSupportedImpl(TSubclassOf<class UMovieScene
 		InTrackClass == UMovieSceneSpawnTrack::StaticClass() ||
 		InTrackClass == UMovieSceneSubTrack::StaticClass() ||
 		InTrackClass == UMovieSceneCVarTrack::StaticClass() ||
-		InTrackClass == UMovieSceneBindingLifetimeTrack::StaticClass())
+		InTrackClass == UMovieSceneBindingLifetimeTrack::StaticClass() ||
+		InTrackClass == UMovieSceneTimeWarpTrack::StaticClass())
 	{
 		return ETrackSupport::Supported;
 	}
diff --git a/Engine/Source/Runtime/UMG/Private/Animation/WidgetAnimation.cpp b/Engine/Source/Runtime/UMG/Private/Animation/WidgetAnimation.cpp
--- a/Engine/Source/Runtime/UMG/Private/Animation/WidgetAnimation.cpp
+++ b/Engine/Source/Runtime/UMG/Private/Animation/WidgetAnimation.cpp
@@ -10,6 +10,7 @@
 #include "Tracks/MovieSceneAudioTrack.h"
 #include "Tracks/MovieSceneEventTrack.h"
 #include "Tracks/MovieSceneMaterialParameterCollectionTrack.h"
+#include "Tracks/MovieSceneTimeWarpTrack.h"
 #include "UObject/SequencerObjectVersion.h"
 
 #if WITH_EDITOR
@@ -116,7 +117,8 @@ ETrackSupport UWidgetAnimation::IsTrackSupportedImpl(TSubclassOf<class UMovieSce
 {
 	if (InTrackClass == UMovieSceneAudioTrack::StaticClass() ||
 		InTrackClass == UMovieSceneEventTrack::StaticClass() ||
-		InTrackClass == UMovieSceneMaterialParameterCollectionTrack::StaticClass())
+		InTrackClass == UMovieSceneMaterialParameterCollectionTrack::StaticClass() ||
+		InTrackClass == UMovieSceneTimeWarpTrack::StaticClass())
 	{
 		return ETrackSupport::Supported;
 	}
