Fix bad access to the SelectabilityBridge when there are no subobject for the Blueprint.

diff --git a/Engine/Plugins/MovieScene/ActorSequence/Source/ActorSequenceEditor/Private/ActorSequenceEditorTabSummoner.cpp b/Engine/Plugins/MovieScene/ActorSequence/Source/ActorSequenceEditor/Private/ActorSequenceEditorTabSummoner.cpp
--- a/Engine/Plugins/MovieScene/ActorSequence/Source/ActorSequenceEditor/Private/ActorSequenceEditorTabSummoner.cpp
+++ b/Engine/Plugins/MovieScene/ActorSequence/Source/ActorSequenceEditor/Private/ActorSequenceEditorTabSummoner.cpp
@@ -246,10 +246,12 @@ public:
 				}
 			}
 
-			FEditorViewportSelectabilityBridge& SelectabilityBridge = BlueprintEditor->GetViewportSelectabilityBridge();
-			SelectabilityBridge.OnIsViewportSelectionLimited().Unbind();
-			SelectabilityBridge.OnGetIsObjectSelectableInViewport().Unbind();
-			SelectabilityBridge.OnGetViewportSelectionLimitedText().Unbind();
+			if (FEditorViewportSelectabilityBridge* SelectabilityBridge = BlueprintEditor->GetViewportSelectabilityBridge())
+			{
+				SelectabilityBridge->OnIsViewportSelectionLimited().Unbind();
+				SelectabilityBridge->OnGetIsObjectSelectableInViewport().Unbind();
+				SelectabilityBridge->OnGetViewportSelectionLimitedText().Unbind();
+			}
 		}
 
 		GEditor->UnregisterForUndo(this);
@@ -446,10 +448,12 @@ public:
 
 		if (const TSharedPtr<FBlueprintEditor> BlueprintEditor = WeakBlueprintEditor.Pin())
 		{
-			FEditorViewportSelectabilityBridge& SelectabilityBridge = BlueprintEditor->GetViewportSelectabilityBridge();
-			SelectabilityBridge.OnIsViewportSelectionLimited().BindSP(this, &SActorSequenceEditorWidgetImpl::IsViewportSelectionLimited);
-			SelectabilityBridge.OnGetIsObjectSelectableInViewport().BindSP(this, &SActorSequenceEditorWidgetImpl::IsObjectSelectableInViewport);
-			SelectabilityBridge.OnGetViewportSelectionLimitedText().BindSP(this, &SActorSequenceEditorWidgetImpl::GetViewportSelectionLimitedText);
+			if (FEditorViewportSelectabilityBridge* SelectabilityBridge = BlueprintEditor->GetViewportSelectabilityBridge())
+			{
+				SelectabilityBridge->OnIsViewportSelectionLimited().BindSP(this, &SActorSequenceEditorWidgetImpl::IsViewportSelectionLimited);
+				SelectabilityBridge->OnGetIsObjectSelectableInViewport().BindSP(this, &SActorSequenceEditorWidgetImpl::IsObjectSelectableInViewport);
+				SelectabilityBridge->OnGetViewportSelectionLimitedText().BindSP(this, &SActorSequenceEditorWidgetImpl::GetViewportSelectionLimitedText);
+			}
 		}
 	}
 
@@ -690,27 +694,28 @@ public:
 			return;
 		}
 
-		FEditorViewportSelectabilityBridge& SelectabilityBridge = BlueprintEditor->GetViewportSelectabilityBridge();
-
-		if (SelectabilityBridge.IsViewportSelectionLimited())
+		if (FEditorViewportSelectabilityBridge* SelectabilityBridge = BlueprintEditor->GetViewportSelectabilityBridge())
 		{
-			// Filter out non-selectable nodes from the selection
-			TSet<FSubobjectEditorTreeNodePtrType> NewSelectedNodes;
-			for (const FSubobjectEditorTreeNodePtrType& Node : SubojectEditor->GetSelectedNodes())
+			if (SelectabilityBridge->IsViewportSelectionLimited())
 			{
-				const FSubobjectData* const Data = Node->GetDataSource();
-				const UActorComponent* const Component = Data ? Data->FindComponentInstanceInActor(GetPreviewActor()) : nullptr;
-				if (IsValid(Component) && IsObjectSelectableInViewport(const_cast<UActorComponent*>(Component)))
+				// Filter out non-selectable nodes from the selection
+				TSet<FSubobjectEditorTreeNodePtrType> NewSelectedNodes;
+				for (const FSubobjectEditorTreeNodePtrType& Node : SubojectEditor->GetSelectedNodes())
 				{
-					NewSelectedNodes.Add(Node);
+					const FSubobjectData* const Data = Node->GetDataSource();
+					const UActorComponent* const Component = Data ? Data->FindComponentInstanceInActor(GetPreviewActor()) : nullptr;
+					if (IsValid(Component) && IsObjectSelectableInViewport(const_cast<UActorComponent*>(Component)))
+					{
+						NewSelectedNodes.Add(Node);
+					}
 				}
-			}
 
-			// Clear selection and select the new nodes
-			SubojectEditor->ClearSelection();
-			for (const FSubobjectEditorTreeNodePtrType& Node : NewSelectedNodes)
-			{
-				SubojectEditor->SelectNode(Node, false);
+				// Clear selection and select the new nodes
+				SubojectEditor->ClearSelection();
+				for (const FSubobjectEditorTreeNodePtrType& Node : NewSelectedNodes)
+				{
+					SubojectEditor->SelectNode(Node, false);
+				}
 			}
 		}
 	}
diff --git a/Engine/Source/Editor/Kismet/Private/BlueprintEditor.cpp b/Engine/Source/Editor/Kismet/Private/BlueprintEditor.cpp
--- a/Engine/Source/Editor/Kismet/Private/BlueprintEditor.cpp
+++ b/Engine/Source/Editor/Kismet/Private/BlueprintEditor.cpp
@@ -10890,10 +10890,9 @@ bool FBlueprintEditor::AreDelegatesAllowed() const
 	return true;
 }
 
-FEditorViewportSelectabilityBridge& FBlueprintEditor::GetViewportSelectabilityBridge()
+FEditorViewportSelectabilityBridge* FBlueprintEditor::GetViewportSelectabilityBridge()
 {
-	check(ViewportSelectabilityBridge.IsValid());
-	return *ViewportSelectabilityBridge.Get();
+	return ViewportSelectabilityBridge.Get();
 }
 
 /////////////////////////////////////////////////////
diff --git a/Engine/Source/Editor/Kismet/Private/SCSEditorViewportClient.cpp b/Engine/Source/Editor/Kismet/Private/SCSEditorViewportClient.cpp
--- a/Engine/Source/Editor/Kismet/Private/SCSEditorViewportClient.cpp
+++ b/Engine/Source/Editor/Kismet/Private/SCSEditorViewportClient.cpp
@@ -344,11 +344,13 @@ void FSCSEditorViewportClient::DrawCanvas( FViewport& InViewport, FSceneView& Vi
 	const TSharedPtr<FBlueprintEditor> BlueprintEditor = BlueprintEditorPtr.Pin();
 	if (BlueprintEditor.IsValid())
 	{
-		FEditorViewportSelectabilityBridge& SelectabilityBridge = BlueprintEditor->GetViewportSelectabilityBridge();
-		if (SelectabilityBridge.IsViewportSelectionLimited())
+		if (FEditorViewportSelectabilityBridge* SelectabilityBridge = BlueprintEditor->GetViewportSelectabilityBridge())
 		{
-			const FText SelectionLimitedText = SelectabilityBridge.GetViewportSelectionLimitedText();
-			FEditorViewportSelectability::DrawEnabledTextNotice(&Canvas, SelectionLimitedText);
+			if (SelectabilityBridge->IsViewportSelectionLimited())
+			{
+				const FText SelectionLimitedText = SelectabilityBridge->GetViewportSelectionLimitedText();
+				FEditorViewportSelectability::DrawEnabledTextNotice(&Canvas, SelectionLimitedText);
+			}
 		}
 	}
 }
@@ -1198,7 +1200,10 @@ bool FSCSEditorViewportClient::IsViewportSelectionLimited() const
 {
 	if (const TSharedPtr<FBlueprintEditor> BlueprintEditor = BlueprintEditorPtr.Pin())
 	{
-		return BlueprintEditor->GetViewportSelectabilityBridge().IsViewportSelectionLimited();
+		if (FEditorViewportSelectabilityBridge* SelectabilityBridge = BlueprintEditor->GetViewportSelectabilityBridge())
+		{
+			return SelectabilityBridge->IsViewportSelectionLimited();
+		}
 	}
 	return false;
 }
@@ -1207,7 +1212,10 @@ bool FSCSEditorViewportClient::IsObjectSelectableInViewport(UObject* const InObj
 {
 	if (const TSharedPtr<FBlueprintEditor> BlueprintEditor = BlueprintEditorPtr.Pin())
 	{
-		return BlueprintEditor->GetViewportSelectabilityBridge().IsObjectSelectableInViewport(InObject);
+		if (FEditorViewportSelectabilityBridge* SelectabilityBridge = BlueprintEditor->GetViewportSelectabilityBridge())
+		{
+			return SelectabilityBridge->IsObjectSelectableInViewport(InObject);
+		}
 	}
 	return true;
 }
diff --git a/Engine/Source/Editor/Kismet/Public/BlueprintEditor.h b/Engine/Source/Editor/Kismet/Public/BlueprintEditor.h
--- a/Engine/Source/Editor/Kismet/Public/BlueprintEditor.h
+++ b/Engine/Source/Editor/Kismet/Public/BlueprintEditor.h
@@ -1386,7 +1386,7 @@ public://@TODO
 	/** Whether delegates are allowed  to be displayed/created in a blueprint */
 	virtual bool AreDelegatesAllowed() const;
 
-	FEditorViewportSelectabilityBridge& GetViewportSelectabilityBridge();
+	FEditorViewportSelectabilityBridge* GetViewportSelectabilityBridge();
 
 protected:
 
