PR #13519: [Sequences] Add PlayReverse function to actor sequence components

diff --git a/Engine/Plugins/MovieScene/ActorSequence/Source/ActorSequence/Private/ActorSequenceComponent.cpp b/Engine/Plugins/MovieScene/ActorSequence/Source/ActorSequence/Private/ActorSequenceComponent.cpp
--- a/Engine/Plugins/MovieScene/ActorSequence/Source/ActorSequence/Private/ActorSequenceComponent.cpp
+++ b/Engine/Plugins/MovieScene/ActorSequence/Source/ActorSequence/Private/ActorSequenceComponent.cpp
@@ -26,6 +26,14 @@ void UActorSequenceComponent::PlaySequence()
 	}
 }
 
+void UActorSequenceComponent::PlaySequenceReverse()
+{
+	if (SequencePlayer != nullptr)
+	{
+		SequencePlayer->PlayReverse();
+	}
+}
+
 void UActorSequenceComponent::PauseSequence()
 {
 	if (SequencePlayer != nullptr)
diff --git a/Engine/Plugins/MovieScene/ActorSequence/Source/ActorSequence/Public/ActorSequenceComponent.h b/Engine/Plugins/MovieScene/ActorSequence/Source/ActorSequence/Public/ActorSequenceComponent.h
--- a/Engine/Plugins/MovieScene/ActorSequence/Source/ActorSequence/Public/ActorSequenceComponent.h
+++ b/Engine/Plugins/MovieScene/ActorSequence/Source/ActorSequence/Public/ActorSequenceComponent.h
@@ -39,6 +39,10 @@ public:
 	UFUNCTION(BlueprintCallable, Category = "Sequencer|Player")
 	UE_API void PlaySequence();
 
+	/** Calls the PlayReverse function on the SequencePlayer if its valid. */
+	UFUNCTION(BlueprintCallable, Category = "Sequencer|Player")
+	UE_API void PlaySequenceReverse();
+
 	/** Calls the Pause function on the SequencePlayer if its valid. */
 	UFUNCTION(BlueprintCallable, Category = "Sequencer|Player")
 	UE_API void PauseSequence();
