Sequencer: Limit Viewport Selection UX Tweaks

- The Limit Viewport Selection button in the viewport menu should now be visible if there is an active sequence (rather than tied to Animation Mode only)
- The Sequencer toolbar doesn't need to have the Limit Viewport Selection button since it's in the viewport menu now.
- For UMG and Actor Sequence, the Selection Lmiited button is now in the toolbar and the text has been removed

diff --git a/Engine/Plugins/MovieScene/ActorSequence/Source/ActorSequenceEditor/ActorSequenceEditor.Build.cs b/Engine/Plugins/MovieScene/ActorSequence/Source/ActorSequenceEditor/ActorSequenceEditor.Build.cs
--- a/Engine/Plugins/MovieScene/ActorSequence/Source/ActorSequenceEditor/ActorSequenceEditor.Build.cs
+++ b/Engine/Plugins/MovieScene/ActorSequence/Source/ActorSequenceEditor/ActorSequenceEditor.Build.cs
@@ -25,7 +25,8 @@ public class ActorSequenceEditor : ModuleRules
 				"SubobjectEditor",
 				"SubobjectDataInterface",
 				"UnrealEd",
-				"TimeManagement"
+				"TimeManagement",
+				"ToolMenus"
 			}
 		);
 
diff --git a/Engine/Plugins/MovieScene/ActorSequence/Source/ActorSequenceEditor/Private/ActorSequenceEditorTabSummoner.cpp b/Engine/Plugins/MovieScene/ActorSequence/Source/ActorSequenceEditor/Private/ActorSequenceEditorTabSummoner.cpp
--- a/Engine/Plugins/MovieScene/ActorSequence/Source/ActorSequenceEditor/Private/ActorSequenceEditorTabSummoner.cpp
+++ b/Engine/Plugins/MovieScene/ActorSequence/Source/ActorSequenceEditor/Private/ActorSequenceEditorTabSummoner.cpp
@@ -31,6 +31,43 @@
 DECLARE_DELEGATE_OneParam(FOnComponentSelected, TSharedPtr<FSCSEditorTreeNode>);
 DECLARE_DELEGATE_RetVal_OneParam(bool, FIsComponentValid, UActorComponent*);
 
+namespace UE::ActorSequenceEditor::Private
+{
+	const FName ViewportToolbarOwnerName = "ActorSequenceEditorViewportToolbar";
+
+	FToolMenuEntry CreateToggleViewportSelectionEntry(TWeakPtr<ISequencer> InWeakSequencer)
+	{
+		return FToolMenuEntry::InitDynamicEntry(
+			"DynamicToggleViewportSelection",
+			FNewToolMenuSectionDelegate::CreateLambda(
+				[InWeakSequencer](FToolMenuSection& InDynamicSection)
+				{
+					// First attempt to use the non-context Sequencer.
+					TSharedPtr<ISequencer> Sequencer = InWeakSequencer.Pin();
+					if (!Sequencer)
+					{
+						return;
+					}
+
+					if (!Sequencer->GetHostCapabilities().bSupportsViewportSelectability)
+					{
+						return;
+					}
+
+					const TSharedPtr<FUICommandList> CommandList =
+						Sequencer->GetCommandBindings(ESequencerCommandBindings::Sequencer);
+
+					FToolMenuEntry& Entry = InDynamicSection.AddMenuEntryWithCommandList(
+						FSequencerCommands::Get().ToggleLimitViewportSelection, CommandList
+					);
+					Entry.InsertPosition.Position = EToolMenuInsertType::Last;
+					Entry.SetShowInToolbarTopLevel(true);
+				}
+			)
+		);
+	}
+
+}; // namespace UE::ActorSequenceEditor::Private
 
 class SComponentSelectionTree
 	: public SCompoundWidget
@@ -250,8 +287,10 @@ public:
 			{
 				SelectabilityBridge->OnIsViewportSelectionLimited().Unbind();
 				SelectabilityBridge->OnGetIsObjectSelectableInViewport().Unbind();
-				SelectabilityBridge->OnGetViewportSelectionLimitedText().Unbind();
 			}
+
+			// Remove the viewport toolbar extensions we added earlier.
+			UToolMenus::Get()->UnregisterOwnerByName(UE::ActorSequenceEditor::Private::ViewportToolbarOwnerName);
 		}
 
 		GEditor->UnregisterForUndo(this);
@@ -452,7 +491,23 @@ public:
 			{
 				SelectabilityBridge->OnIsViewportSelectionLimited().BindSP(this, &SActorSequenceEditorWidgetImpl::IsViewportSelectionLimited);
 				SelectabilityBridge->OnGetIsObjectSelectableInViewport().BindSP(this, &SActorSequenceEditorWidgetImpl::IsObjectSelectableInViewport);
-				SelectabilityBridge->OnGetViewportSelectionLimitedText().BindSP(this, &SActorSequenceEditorWidgetImpl::GetViewportSelectionLimitedText);
+			}
+		}
+
+		// Extend the SCSEditor viewport toolbar.
+		{
+			FToolMenuOwnerScoped ScopeOwner(UE::ActorSequenceEditor::Private::ViewportToolbarOwnerName);
+
+			const FName TransformMenu = "SCSEditor.ViewportToolbar.Transform";
+			UToolMenu* const Menu = UToolMenus::Get()->ExtendMenu(TransformMenu);
+
+			FToolMenuSection& SelectionSection = Menu->FindOrAddSection("Selection");
+
+			{
+				FToolMenuEntry Entry = UE::ActorSequenceEditor::Private::CreateToggleViewportSelectionEntry(Sequencer.ToWeakPtr());
+				Entry.InsertPosition.Position = EToolMenuInsertType::Last;
+
+				SelectionSection.AddEntry(Entry);
 			}
 		}
 	}
@@ -720,17 +775,6 @@ public:
 		}
 	}
 
-	FText GetViewportSelectionLimitedText() const
-	{
-		if (Sequencer.IsValid())
-		{
-			const FSequencerCommands& SequencerCommands = FSequencerCommands::Get();
-			return FEditorViewportSelectability::GetLimitedSelectionText(SequencerCommands.ToggleLimitViewportSelection
-				, LOCTEXT("SequencerSelectionLimitHelpText", "Sequencer Selection Limited"));
-		}
-		return FText();
-	}
-
 private:
 	TWeakObjectPtr<UActorSequence> WeakSequence;
 
diff --git a/Engine/Plugins/MovieScene/LevelSequenceEditor/Source/LevelSequenceEditor/Private/LevelSequenceEditorToolkit.cpp b/Engine/Plugins/MovieScene/LevelSequenceEditor/Source/LevelSequenceEditor/Private/LevelSequenceEditorToolkit.cpp
--- a/Engine/Plugins/MovieScene/LevelSequenceEditor/Source/LevelSequenceEditor/Private/LevelSequenceEditorToolkit.cpp
+++ b/Engine/Plugins/MovieScene/LevelSequenceEditor/Source/LevelSequenceEditor/Private/LevelSequenceEditorToolkit.cpp
@@ -87,35 +87,7 @@ FToolMenuEntry CreateToggleViewportSelectionEntry(TWeakPtr<ISequencer> InWeakSeq
 					FSequencerCommands::Get().ToggleLimitViewportSelection, CommandList
 				);
 				Entry.InsertPosition.Position = EToolMenuInsertType::Last;
-
-				if (UUnrealEdViewportToolbarContext* ViewportToolbarContext =
-						InDynamicSection.FindContext<UUnrealEdViewportToolbarContext>())
-				{
-					// Show the entry in the top-level toolbar if we're in Animation Mode (EditMode.ControlRig). Otherwise it will be in the Transform submenu.
-					Entry.SetShowInToolbarTopLevel(TAttribute<bool>::CreateLambda(
-						[WeakViewport = ViewportToolbarContext->Viewport]() -> bool
-						{
-							if (TSharedPtr<SEditorViewport> Viewport = WeakViewport.Pin())
-							{
-								if (TSharedPtr<FEditorViewportClient> Client = Viewport->GetViewportClient())
-								{
-									if (FEditorModeTools* ModeTools = Client->GetModeTools())
-									{
-										// Hard-code the mode name instead of using FControlRigEditMode::ModeName to avoid adding a dependency on ControlRigEditor.
-										const FName ControlRigModeName = "EditMode.ControlRig";
-
-										if (ModeTools->GetActiveMode(ControlRigModeName))
-										{
-											return true;
-										}
-									}
-								}
-							}
-
-							return false;
-						}
-					));
-				}
+				Entry.SetShowInToolbarTopLevel(true);
 			}
 		)
 	);
diff --git a/Engine/Source/Editor/Sequencer/Private/SSequencer.cpp b/Engine/Source/Editor/Sequencer/Private/SSequencer.cpp
--- a/Engine/Source/Editor/Sequencer/Private/SSequencer.cpp
+++ b/Engine/Source/Editor/Sequencer/Private/SSequencer.cpp
@@ -1638,20 +1638,6 @@ void SSequencer::PopulateToolBar(UToolMenu* InMenu)
 		Section.AddSeparator(NAME_None);
 	}
 
-	{
-		if (Sequencer->GetHostCapabilities().bSupportsViewportSelectability)
-		{
-			FToolMenuSection& SelectionSection = InMenu->AddSection(TEXT("Selection"));
-
-			FToolMenuEntry ToggleLockViewportSelectionEntry =
-				FToolMenuEntry::InitToolBarButton(FSequencerCommands::Get().ToggleLimitViewportSelection);
-			ToggleLockViewportSelectionEntry.StyleNameOverride = SequencerToolbarStyleName;
-			SelectionSection.AddEntry(ToggleLockViewportSelectionEntry);
-
-			SelectionSection.AddSeparator(NAME_None);
-		}
-	}
-
 	{
 		FToolMenuSection& SnappingSection = InMenu->AddSection("Snapping");
 
diff --git a/Engine/Source/Editor/UMGEditor/Private/Designer/SDesignerToolBar.cpp b/Engine/Source/Editor/UMGEditor/Private/Designer/SDesignerToolBar.cpp
--- a/Engine/Source/Editor/UMGEditor/Private/Designer/SDesignerToolBar.cpp
+++ b/Engine/Source/Editor/UMGEditor/Private/Designer/SDesignerToolBar.cpp
@@ -20,6 +20,7 @@
 #include "Designer/DesignerCommands.h"
 #include "SViewportToolBarComboMenu.h"
 #include "ViewportToolbar/UnrealEdViewportToolbar.h"
+#include "SequencerCommands.h"
 
 #define LOCTEXT_NAMESPACE "UMG"
 
@@ -69,6 +70,7 @@ TSharedRef< SWidget > SDesignerToolBar::MakeToolBar(const TSharedPtr< FExtender
 	ToolbarBuilder.BeginBlockGroup();
 	ToolbarBuilder.AddToolBarButton(FDesignerCommands::Get().ToggleOutlines, NAME_None, TAttribute<FText>(), TAttribute<FText>(), TAttribute<FSlateIcon>(), "ToggleOutlines");
 	ToolbarBuilder.AddToolBarButton(FDesignerCommands::Get().ToggleRespectLocks, NAME_None, TAttribute<FText>(), TAttribute<FText>(), TAttribute<FSlateIcon>(), "ToggleRespectLocks");
+	ToolbarBuilder.AddToolBarButton(FSequencerCommands::Get().ToggleLimitViewportSelection, NAME_None, TAttribute<FText>(), TAttribute<FText>(), TAttribute<FSlateIcon>(), "ToggleLimitViewportSelection");
 	ToolbarBuilder.EndBlockGroup();
 	ToolbarBuilder.EndSection();
 
diff --git a/Engine/Source/Editor/UMGEditor/Private/Designer/SDesignerView.cpp b/Engine/Source/Editor/UMGEditor/Private/Designer/SDesignerView.cpp
--- a/Engine/Source/Editor/UMGEditor/Private/Designer/SDesignerView.cpp
+++ b/Engine/Source/Editor/UMGEditor/Private/Designer/SDesignerView.cpp
@@ -81,6 +81,7 @@
 #include "Styling/ToolBarStyle.h"
 #include "UMGEditorProjectSettings.h"
 #include "Types/InvisibleToWidgetReflectorMetaData.h"
+#include "SequencerCommands.h"
 
 #define LOCTEXT_NAMESPACE "UMG"
 
@@ -481,15 +482,6 @@ TSharedRef<SWidget> SDesignerView::CreateOverlayUI()
 				.ColorAndOpacity(this, &SDesignerView::GetZoomTextColorAndOpacity)
 				.Visibility(EVisibility::SelfHitTestInvisible)
 			]
-			+ SVerticalBox::Slot()
-			.AutoHeight()
-			[
-				SNew(STextBlock)
-				.TextStyle(FAppStyle::Get(), TEXT("Graph.ZoomText"))
-				.Text(LOCTEXT("SequencerSelectionLimited", "Sequencer Selection Limited"))
-				.ColorAndOpacity(FLinearColor(1.f, 1.f, 1.f, 0.25f))
-				.Visibility(this, &SDesignerView::GetSelectionLimitedTextVisibility)
-			]
 		]
 
 		+ SHorizontalBox::Slot()
@@ -847,6 +839,13 @@ void SDesignerView::BindCommands()
 		FCanExecuteAction(),
 		FIsActionChecked::CreateSP(this, &SDesignerView::IsRespectingLocks)
 	);
+
+	CommandList->MapAction(
+		FSequencerCommands::Get().ToggleLimitViewportSelection,
+		FExecuteAction::CreateSP(this, &SDesignerView::ToggleViewportSelectionLimited),
+		FCanExecuteAction::CreateSP(this, &SDesignerView::CanToggleViewportSelectionLimited),
+		FIsActionChecked::CreateSP(this, &SDesignerView::IsViewportSelectionLimited)
+	);
 }
 
 void SDesignerView::AddReferencedObjects(FReferenceCollector& Collector)
@@ -3989,22 +3988,55 @@ void SDesignerView::DeselectNonSequencerWidgets()
 	WidgetBlueprintEditor->SelectWidgets(NewSelection, false);
 }
 
-EVisibility SDesignerView::GetSelectionLimitedTextVisibility() const
+void SDesignerView::ToggleViewportSelectionLimited()
 {
 	const TSharedPtr<FWidgetBlueprintEditor> WidgetBlueprintEditor = BlueprintEditor.Pin();
-	if (!WidgetBlueprintEditor.IsValid())
+	if (!WidgetBlueprintEditor.IsValid() || WidgetBlueprintEditor->GetCurrentAnimation() == UWidgetAnimation::GetNullAnimation())
+	{
+		return;
+	}
+
+	const TSharedPtr<ISequencer>& ActiveSequencer = WidgetBlueprintEditor->GetSequencer();
+	if (!ActiveSequencer.IsValid())
+	{
+		return;
+	}
+
+	ActiveSequencer->SetViewportSelectionLimited(!ActiveSequencer->IsViewportSelectionLimited());
+}
+
+bool SDesignerView::CanToggleViewportSelectionLimited() const
+{
+	const TSharedPtr<FWidgetBlueprintEditor> WidgetBlueprintEditor = BlueprintEditor.Pin();
+	if (!WidgetBlueprintEditor.IsValid() || WidgetBlueprintEditor->GetCurrentAnimation() == UWidgetAnimation::GetNullAnimation())
 	{
-		return EVisibility::Collapsed;
+		return false;
 	}
 
 	const TSharedPtr<ISequencer>& ActiveSequencer = WidgetBlueprintEditor->GetSequencer();
-	if (!ActiveSequencer.IsValid() || !ActiveSequencer->IsViewportSelectionLimited())
+	if (!ActiveSequencer.IsValid())
+	{
+		return false;
+	}
+
+	return true;
+}
+
+bool SDesignerView::IsViewportSelectionLimited() const
+{
+	const TSharedPtr<FWidgetBlueprintEditor> WidgetBlueprintEditor = BlueprintEditor.Pin();
+	if (!WidgetBlueprintEditor.IsValid() || WidgetBlueprintEditor->GetCurrentAnimation() == UWidgetAnimation::GetNullAnimation())
+	{
+		return false;
+	}
+
+	const TSharedPtr<ISequencer>& ActiveSequencer = WidgetBlueprintEditor->GetSequencer();
+	if (!ActiveSequencer.IsValid())
 	{
-		return EVisibility::Collapsed;
+		return false;
 	}
 
-	return (WidgetBlueprintEditor->GetCurrentAnimation() == UWidgetAnimation::GetNullAnimation())
-		? EVisibility::Collapsed : EVisibility::SelfHitTestInvisible;
+	return ActiveSequencer->IsViewportSelectionLimited();
 }
 
 #undef LOCTEXT_NAMESPACE
diff --git a/Engine/Source/Editor/UMGEditor/Public/Designer/SDesignerView.h b/Engine/Source/Editor/UMGEditor/Public/Designer/SDesignerView.h
--- a/Engine/Source/Editor/UMGEditor/Public/Designer/SDesignerView.h
+++ b/Engine/Source/Editor/UMGEditor/Public/Designer/SDesignerView.h
@@ -309,7 +309,11 @@ private:
 
 	UE_API void DeselectNonSequencerWidgets();
 
-	UE_API EVisibility GetSelectionLimitedTextVisibility() const;
+	UE_API void ToggleViewportSelectionLimited();
+	
+	UE_API bool CanToggleViewportSelectionLimited() const;
+
+	UE_API bool IsViewportSelectionLimited() const;
 
 private:
 	/** A reference to the BP Editor that owns this designer */
